Set the stand-alone virtualization tool to always build with the latest include order

#rnx

diff --git a/Engine/Source/Programs/UnrealVirtualizationTool/UnrealVirtualizationTool.Target.cs b/Engine/Source/Programs/UnrealVirtualizationTool/UnrealVirtualizationTool.Target.cs
--- a/Engine/Source/Programs/UnrealVirtualizationTool/UnrealVirtualizationTool.Target.cs
+++ b/Engine/Source/Programs/UnrealVirtualizationTool/UnrealVirtualizationTool.Target.cs
@@ -9,6 +9,7 @@ public class UnrealVirtualizationToolTarget : TargetRules
 	{
 		Type = TargetType.Program;
 		LinkType = TargetLinkType.Monolithic;
+		IncludeOrderVersion = EngineIncludeOrderVersion.Latest;
 
 		LaunchModuleName = "UnrealVirtualizationTool";
 
