Extend the VA stand-alone tool to work on both changelists (the existing functionality) and lists of packages stored in a text file (the new functionality)

#rnx

### Rationale
- The new mode being added allows us to write a list of package file paths to be virtualized to a file on disk and then the file to be passed to the tool so that it knows what to virtualize. This new mode will then be used by P4VUtils.dll which is how we will run the virtualization process via p4v custom tools.
- We need this second mode because:
-- A workspace could have multiple projects in it, which means a user could try to virtualize a changelist with packages from more than one project.
-- Projects can be associated with different engine installations so the user could try to virtualize a changelist with packages using different engine installations, which in turn means that they need to be virtualized by different versions of the stand-alone tool.
-- This means we cannot rely on passing the CL number to the tool, install P4VUtils.dll will need to sort the package paths by engine installation, then pass those paths to the correct version of the tool.
-- The paths need to be written to a text file to be passed to the tool as the total length of path data can easily exceed the commandline max length limit if someone were to try and virtualize an entire project in one go for example.
- Note that the new mode does NOT submit the package files since there might be multiple calls to different versions of the stand-alone tool. In theory if all the packages are under one engine install (which is almost certainly going to be 99.999% of all submits) we could allow the tool to do the submit but that means we need to maintain two different code paths for the same operation.
- The existing mode (virtualize and submit a changelist) remains unchanged. It is not likely that an end user will ever use this mode, but it will make for a useful debugging tool in the future.

### Changes
- When the tool first starts we now print the log path to the user so that they can more easily find it. This might be something that we want to remove after a while, but during the initial release of the tool it should make getting good debug info back from users with problems somewhat easier.
- Some logging and error messages have had small updates made, to make the final output to the user more readable (based on the many many times this tool was run while testing). In some places additional logging has been added where I felt it was not clear to the user as to what was going on.
- Connecting to source control has been moved from the changelist parsing code to it's own method ::TryConnectToSourceControl since all modes will require it.
- The cmdline help output has been redone to support the multiple command line options we now have.
- Although we could infer the mode to run based on the commandline options, we require the users to explicitly state the mode via -Mode=XXX. In general the tool will be invoked from a gui so adding extra commands doesn't really affect the usability of the system and it makes things a bit clearer.

diff --git a/Engine/Source/Programs/UnrealVirtualizationTool/Private/UnrealVirtualizationTool.cpp b/Engine/Source/Programs/UnrealVirtualizationTool/Private/UnrealVirtualizationTool.cpp
--- a/Engine/Source/Programs/UnrealVirtualizationTool/Private/UnrealVirtualizationTool.cpp
+++ b/Engine/Source/Programs/UnrealVirtualizationTool/Private/UnrealVirtualizationTool.cpp
@@ -26,7 +26,7 @@ bool UnrealVirtualizationToolMain(int32 ArgC, TCHAR* ArgV[])
 	{
 		if (!App.Run())
 		{
-			UE_LOG(LogVirtualizationTool, Error, TEXT("UnrealVirtualizationTool failed to run"));
+			UE_LOG(LogVirtualizationTool, Error, TEXT("UnrealVirtualizationTool ran with errors"));
 			bRanSuccessfully = false;
 		}
 	}	
diff --git a/Engine/Source/Programs/UnrealVirtualizationTool/Private/UnrealVirtualizationToolApp.cpp b/Engine/Source/Programs/UnrealVirtualizationTool/Private/UnrealVirtualizationToolApp.cpp
--- a/Engine/Source/Programs/UnrealVirtualizationTool/Private/UnrealVirtualizationToolApp.cpp
+++ b/Engine/Source/Programs/UnrealVirtualizationTool/Private/UnrealVirtualizationToolApp.cpp
@@ -2,6 +2,7 @@
 
 #include "UnrealVirtualizationToolApp.h"
 
+#include "GenericPlatform/GenericPlatformOutputDevices.h"
 #include "HAL/FeedbackContextAnsi.h"
 #include "HAL/FileManager.h"
 #include "Interfaces/IPluginManager.h"
@@ -9,6 +10,7 @@
 #include "ISourceControlProvider.h"
 #include "Misc/CommandLine.h"
 #include "Misc/ConfigCacheIni.h"
+#include "Misc/FileHelper.h"
 #include "Misc/Parse.h"
 #include "Misc/Paths.h"
 #include "Misc/PathViews.h"
@@ -132,6 +134,23 @@ void ReloadConfigs(FConfigFile& PluginConfig)
 	}
 }
 
+/** Utility to get EMode from a string */
+void LexFromString(EMode& OutValue, const FStringView& InString)
+{
+	if (InString == TEXT("Changelist"))
+	{
+		OutValue = EMode::Changelist;
+	}
+	else if (InString == TEXT("PackageList"))
+	{
+		OutValue = EMode::PackageList;
+	}
+	else
+	{
+		OutValue = EMode::Unknown;
+	}
+}
+
 } // namespace
 
 /** 
@@ -183,6 +202,13 @@ EInitResult FUnrealVirtualizationToolApp::Initialize()
 
 	UE_LOG(LogVirtualizationTool, Display, TEXT("Initializing..."));
 
+	// Display the log path to the user so that they can more easily find it
+	// Note that ::GetAbsoluteLogFilename does not always return an absolute filename
+	FString LogFilePath = FGenericPlatformOutputDevices::GetAbsoluteLogFilename();
+	LogFilePath = IFileManager::Get().ConvertToAbsolutePathForExternalAppForRead(*LogFilePath);
+
+	UE_LOG(LogVirtualizationTool, Display, TEXT("Logging process to '%s'"), *LogFilePath);
+
 	EInitResult CmdLineResult = TryParseCmdLine();
 	if (CmdLineResult != EInitResult::Success)
 	{
@@ -199,14 +225,35 @@ EInitResult FUnrealVirtualizationToolApp::Initialize()
 		return EInitResult::Error;
 	}
 
-	// Parse the provided changelist
-	TArray<FString> ChangelistPackages;
-	if (!TryParseChangelist(ChangelistPackages))
+	if (!TryConnectToSourceControl())
 	{
 		return EInitResult::Error;
 	}
 
-	if (!TrySortFilesByProject(ChangelistPackages))
+	TArray<FString> Packages;
+	switch (Mode)
+	{
+		case EMode::Changelist:
+			if (!TryParseChangelist(Packages))
+			{
+				return EInitResult::Error;
+			}
+			break;
+
+		case EMode::PackageList:
+			if (!TryParsePackageList(Packages))
+			{
+				return EInitResult::Error;
+			}
+			break;
+
+		default:
+			UE_LOG(LogVirtualizationTool, Display, TEXT("Unknown mode, cannot find packages!"));
+			return EInitResult::Error;
+			break;
+	}
+
+	if (!TrySortFilesByProject(Packages))
 	{
 		return EInitResult::Error;
 	}
@@ -228,7 +275,7 @@ bool FUnrealVirtualizationToolApp::Run()
 
 		if (!SCCProvider.IsValid())
 		{
-			UE_LOG(LogVirtualizationTool, Error, TEXT("No valid sourcer control connection found!"));
+			UE_LOG(LogVirtualizationTool, Error, TEXT("No valid source control connection found!"));
 			return false;
 		}
 
@@ -236,6 +283,11 @@ bool FUnrealVirtualizationToolApp::Run()
 
 		for (const FProject& Project : Projects)
 		{
+			TStringBuilder<128> ProjectName;
+			ProjectName  << Project.GetProjectName();
+
+			UE_LOG(LogVirtualizationTool, Display, TEXT("\tChecking package(s) for the project '%s'..."), ProjectName.ToString());
+
 			FConfigFile EngineConfigWithProject;
 			if (!Project.TryLoadConfig(EngineConfigWithProject))
 			{
@@ -244,7 +296,7 @@ bool FUnrealVirtualizationToolApp::Run()
 
 			Project.RegisterMountPoints();
 
-			UE::Virtualization::FInitParams InitParams(Project.GetProjectName(), EngineConfigWithProject);
+			UE::Virtualization::FInitParams InitParams(ProjectName, EngineConfigWithProject);
 			UE::Virtualization::Initialize(InitParams);
 
 			TArray<FString> Packages = Project.GetAllPackages();
@@ -262,6 +314,8 @@ bool FUnrealVirtualizationToolApp::Run()
 				return false;
 			}
 
+			UE_LOG(LogVirtualizationTool, Display, TEXT("\tCheck complete"));
+
 			UE::Virtualization::Shutdown();
 			Project.UnRegisterMountPoints();
 		}
@@ -273,9 +327,12 @@ bool FUnrealVirtualizationToolApp::Run()
 		UE_LOG(LogVirtualizationTool, Display, TEXT("Skipping the virtualization process"));
 	}
 
-	if (!TrySubmitChangelist(FinalDescriptionTags))
+	if (Mode == EMode::Changelist)
 	{
-		return false;
+		if (!TrySubmitChangelist(FinalDescriptionTags))
+		{
+			return false;
+		}
 	}
 
 	return true;
@@ -284,10 +341,11 @@ bool FUnrealVirtualizationToolApp::Run()
 void FUnrealVirtualizationToolApp::PrintCmdLineHelp() const
 {
 	UE_LOG(LogVirtualizationTool, Display, TEXT("Usage:"));
-	UE_LOG(LogVirtualizationTool, Display, TEXT("\tUnrealVirtualizationTool -ClientSpecName=<name> -Changelist=<number>"));
-	UE_LOG(LogVirtualizationTool, Display, TEXT("\tOptions:"));
-	UE_LOG(LogVirtualizationTool, Display, TEXT("\t\t-nosubmit (the changelist will be virtualized but not submitted)"));
-	UE_LOG(LogVirtualizationTool, Display, TEXT("\t\t-verbose (all log messages with display verbosity will be displayed, not just LogVirtualizationTool)"));
+	UE_LOG(LogVirtualizationTool, Display, TEXT("UnrealVirtualizationTool -ClientSpecName=<name> -Mode=Changelist -Changelist=<number> [-nosubmit] [global options]"));
+	UE_LOG(LogVirtualizationTool, Display, TEXT("\t[optional]-nosubmit (the changelist will be virtualized but not submitted)"));
+	UE_LOG(LogVirtualizationTool, Display, TEXT("UnrealVirtualizationTool -ClientSpecName=<name> -Mode=PackageList -Path=<string> [global options]"));
+	UE_LOG(LogVirtualizationTool, Display, TEXT("Global Options:"));
+	UE_LOG(LogVirtualizationTool, Display, TEXT("\t-verbose (all log messages with display verbosity will be displayed, not just LogVirtualizationTool)"));
 }
 
 bool FUnrealVirtualizationToolApp::TrySubmitChangelist(const TArray<FString>& DescriptionTags)
@@ -343,7 +401,10 @@ bool FUnrealVirtualizationToolApp::TrySubmitChangelist(const TArray<FString>& De
 	}
 	else
 	{
-		UE_LOG(LogVirtualizationTool, Error, TEXT("Submit failed, please check the log file!"));
+		// Even when log suppression is active we still show errors to the users and as the source control
+		// operation should have logged the problem as an error the user will see it. This means we don't 
+		// have to extract it from CheckInOperation .
+		UE_LOG(LogVirtualizationTool, Error, TEXT("Submit failed, please check the log!"));
 
 		return false;
 	}
@@ -387,6 +448,26 @@ bool FUnrealVirtualizationToolApp::TryInitEnginePlugins()
 	return true;
 }
 
+bool FUnrealVirtualizationToolApp::TryConnectToSourceControl()
+{
+	TRACE_CPUPROFILER_EVENT_SCOPE(TryConnectToSourceControl);
+
+	UE_LOG(LogVirtualizationTool, Log, TEXT("Trying to connect to source control..."));
+
+	FSourceControlInitSettings SCCSettings(FSourceControlInitSettings::EBehavior::OverrideAll);
+	SCCSettings.AddSetting(TEXT("P4Client"), ClientSpecName);
+
+	SCCProvider = ISourceControlModule::Get().CreateProvider(FName("Perforce"), TEXT("UnrealVirtualizationTool"), SCCSettings);
+	if (SCCProvider.IsValid())
+	{
+		return true;
+	}
+	else
+	{
+		UE_LOG(LogVirtualizationTool, Error, TEXT("Failed to create a perforce connection"));
+		return false;
+	}
+}
 
 EInitResult FUnrealVirtualizationToolApp::TryParseCmdLine()
 {
@@ -410,55 +491,104 @@ EInitResult FUnrealVirtualizationToolApp::TryParseCmdLine()
 		return EInitResult::EarlyOut;
 	}
 
+	// First parse the command line options that can apply to all modes
+
 	if (!FParse::Value(CmdLine, TEXT("-ClientSpecName="), ClientSpecName))
 	{
-		UE_LOG(LogVirtualizationTool, Error, TEXT("Failed to find cmdline 'ClientSpecName', this is a required parameter!"));
+		UE_LOG(LogVirtualizationTool, Error, TEXT("Failed to find cmdline switch 'ClientSpecName', this is a required parameter!"));
 		PrintCmdLineHelp();
 		return EInitResult::Error;
 	}
 
-	if (!FParse::Value(CmdLine, TEXT("-Changelist="), ChangelistNumber))
+	if (FParse::Param(CmdLine, TEXT("Verbose")))
+	{
+		UE_LOG(LogVirtualizationTool, Display, TEXT("Cmdline parameter '-Verbose' found, no longer supressing Display log messages!"));
+	}
+	else
+	{
+		OutputDeviceOverride = MakeUnique<FOverrideOutputDevice>();
+	}
+
+	// Now parse the mode specific command line options
+
+	FString ModeAsString;
+	if (!FParse::Value(CmdLine, TEXT("-Mode="), ModeAsString))
 	{
-		UE_LOG(LogVirtualizationTool, Error, TEXT("Failed to find cmdline 'Changelist', this is a required parameter!"));
+		UE_LOG(LogVirtualizationTool, Error, TEXT("Failed to find cmdline switch 'Mode', this is a required parameter!"));
 		PrintCmdLineHelp();
 		return EInitResult::Error;
 	}
 
-	// Optional switches
-	if (FParse::Param(CmdLine, TEXT("NoSubmit")))
+	LexFromString(Mode, ModeAsString);
+
+	switch (Mode)
 	{
-		UE_LOG(LogVirtualizationTool, Display, TEXT("Cmdline parameter '-NoSubmit' found, the changelist will be virtualized but not submitted!"));
+		case EMode::Changelist:
+			return TryParseChangelistCmdLine(CmdLine);
+			break;
+		case EMode::PackageList:
+			return TryParsePackageListCmdLine(CmdLine);
+			break;
+		case EMode::Unknown:
+		default:
+			UE_LOG(LogVirtualizationTool, Error, TEXT("Unexpected value for the cmdline switch 'Mode', this is a required parameter!"));
+			PrintCmdLineHelp();
+			return EInitResult::Error;
+
+			break;
+	}
+}
+
+EInitResult	FUnrealVirtualizationToolApp::TryParseChangelistCmdLine(const TCHAR* CmdLine)
+{
+	if (FParse::Value(CmdLine, TEXT("-Changelist="), ChangelistNumber))
+	{
+		// Optional switches
+		if (FParse::Param(CmdLine, TEXT("NoSubmit")))
+		{
+			UE_LOG(LogVirtualizationTool, Display, TEXT("Cmdline parameter '-NoSubmit' found, the changelist will be virtualized but not submitted!"));
+		}
+		else
+		{
+			ProcessOptions |= EProcessOptions::Submit;
+		}
+
+		UE_LOG(LogVirtualizationTool, Display, TEXT("Attempting to virtualize changelist '%s'"), *ChangelistNumber);
+		
+		return EInitResult::Success;
 	}
 	else
 	{
-		ProcessOptions |= EProcessOptions::Submit;
+		UE_LOG(LogVirtualizationTool, Error, TEXT("Failed to find cmdline switch 'Changelist', this is a required parameter for the 'Changelist' mode!"));
+		PrintCmdLineHelp();
+		return EInitResult::Error;
 	}
+}
 
-	if (FParse::Param(CmdLine, TEXT("Verbose")))
+EInitResult FUnrealVirtualizationToolApp::TryParsePackageListCmdLine(const TCHAR* CmdLine)
+{
+	if (FParse::Value(CmdLine, TEXT("-Path="), PackageListPath))
 	{
-		UE_LOG(LogVirtualizationTool, Display, TEXT("Cmdline parameter '-Verbose' found, no longer supressing Display log messages!"));
+		UE_LOG(LogVirtualizationTool, Display, TEXT("Attempting to virtualize package list '%s'"), *PackageListPath);
+		return EInitResult::Success;
 	}
 	else
 	{
-		OutputDeviceOverride = MakeUnique<FOverrideOutputDevice>();
-	}
-
-	return EInitResult::Success;
+		UE_LOG(LogVirtualizationTool, Error, TEXT("Failed to find cmdline switch 'Path', this is a required parameter for the 'PackageList mode!"));
+		PrintCmdLineHelp();
+		return EInitResult::Error;
+	}	
 }
 
-bool FUnrealVirtualizationToolApp::TryParseChangelist(TArray<FString>& OutChangelistPackages)
+bool FUnrealVirtualizationToolApp::TryParseChangelist(TArray<FString>& OutPackages)
 {
 	TRACE_CPUPROFILER_EVENT_SCOPE(TryParseChangelist);
 
 	UE_LOG(LogVirtualizationTool, Display, TEXT("Attempting to parse changelist '%s' in workspace '%s'"), *ChangelistNumber, *ClientSpecName);
 
-	FSourceControlInitSettings SCCSettings(FSourceControlInitSettings::EBehavior::OverrideAll);
-	SCCSettings.AddSetting(TEXT("P4Client"), ClientSpecName);
-
-	SCCProvider = ISourceControlModule::Get().CreateProvider(FName("Perforce"), TEXT("UnrealVirtualizationTool"), SCCSettings);
 	if (!SCCProvider.IsValid())
 	{
-		UE_LOG(LogVirtualizationTool, Error, TEXT("Failed to create a perforce connection"));
+		UE_LOG(LogVirtualizationTool, Error, TEXT("No valid source control connection found!"));
 		return false;
 	}
 
@@ -513,7 +643,7 @@ bool FUnrealVirtualizationToolApp::TryParseChangelist(TArray<FString>& OutChange
 			{
 				if (IsPackageFile(FileState->GetFilename()))
 				{
-					OutChangelistPackages.Add(FileState->GetFilename());
+					OutPackages.Add(FileState->GetFilename());
 				}
 				else
 				{
@@ -523,7 +653,7 @@ bool FUnrealVirtualizationToolApp::TryParseChangelist(TArray<FString>& OutChange
 
 			ChangelistToSubmit = Changelist;
 
-			UE_LOG(LogVirtualizationTool, Display, TEXT("Found '%d' package file(s)"), OutChangelistPackages.Num());
+			UE_LOG(LogVirtualizationTool, Display, TEXT("Found '%d' package file(s)"), OutPackages.Num());
 
 			return true;
 		}
@@ -533,13 +663,44 @@ bool FUnrealVirtualizationToolApp::TryParseChangelist(TArray<FString>& OutChange
 	return false;
 }
 
-bool FUnrealVirtualizationToolApp::TrySortFilesByProject(const TArray<FString>& ChangelistPackages)
+bool FUnrealVirtualizationToolApp::TryParsePackageList(TArray<FString>& OutPackages)
+{
+	TRACE_CPUPROFILER_EVENT_SCOPE(TrySortFilesByProject);
+
+	UE_LOG(LogVirtualizationTool, Display, TEXT("Attempting to parse package list '%s'"), *PackageListPath);
+
+	if (!IFileManager::Get().FileExists(*PackageListPath))
+	{
+		UE_LOG(LogVirtualizationTool, Error, TEXT("The package list '%s' does not exist"), *PackageListPath);
+		return false;
+	}
+
+	if (FFileHelper::LoadFileToStringArray(OutPackages, *PackageListPath))
+	{
+		// We don't have control over how the package list was generated so make sure that the paths
+		// are in the format that we want.
+		for (FString& PackagePath : OutPackages)
+		{
+			FPaths::NormalizeFilename(PackagePath);
+		}
+
+		UE_LOG(LogVirtualizationTool, Display, TEXT("Found '%d' package file(s)"), OutPackages.Num());
+		return true;
+	}
+	else
+	{
+		UE_LOG(LogVirtualizationTool, Error, TEXT("Failed to parse the package list '%s'"), *PackageListPath);
+		return false;
+	}
+}
+
+bool FUnrealVirtualizationToolApp::TrySortFilesByProject(const TArray<FString>& Packages)
 {
 	TRACE_CPUPROFILER_EVENT_SCOPE(TrySortFilesByProject);
 
 	UE_LOG(LogVirtualizationTool, Display, TEXT("Sorting files by project"));
 
-	for (const FString& PackagePath : ChangelistPackages)
+	for (const FString& PackagePath : Packages)
 	{
 		FString ProjectFilePath;
 		FString PluginFilePath;
diff --git a/Engine/Source/Programs/UnrealVirtualizationTool/Private/UnrealVirtualizationToolApp.h b/Engine/Source/Programs/UnrealVirtualizationTool/Private/UnrealVirtualizationToolApp.h
--- a/Engine/Source/Programs/UnrealVirtualizationTool/Private/UnrealVirtualizationToolApp.h
+++ b/Engine/Source/Programs/UnrealVirtualizationTool/Private/UnrealVirtualizationToolApp.h
@@ -35,6 +35,17 @@ struct FProject
 	bool TryLoadConfig(FConfigFile& OutConfig) const;
 };
 
+/** The mode of the application to run */
+enum class EMode :uint32
+{
+	/** Error condition */
+	Unknown = 0,
+	/** Virtualize and submit a given changelist */
+	Changelist,
+	/** Virtualize a list of packages provided by text file */
+	PackageList
+};
+
 // Note that the enum doesn't give us huge value right now, it has been provided
 // in case that we expand with additional functionality in the future.
 
@@ -78,9 +89,17 @@ private:
 
 	bool TryLoadModules();
 	bool TryInitEnginePlugins();
+	bool TryConnectToSourceControl();
+
 	EInitResult TryParseCmdLine();
-	bool TryParseChangelist(TArray<FString>& OutChangelistPackages);
-	bool TrySortFilesByProject(const TArray<FString>& ChangelistPackages);
+	EInitResult TryParseChangelistCmdLine(const TCHAR* CmdLine);
+	EInitResult	TryParsePackageListCmdLine(const TCHAR* CmdLine);
+
+	bool TryParseChangelist(TArray<FString>& OutPackages);
+	bool TryParsePackageList(TArray<FString>& OutPackages);
+
+	bool TrySortFilesByProject(const TArray<FString>& Packages);
+
 	bool TryFindProject(const FString& PackagePath, FString& ProjectFilePath, FString& PluginFilePath) const;
 
 	FProject& FindOrAddProject(const FString& ProjectFilePath);
@@ -102,9 +121,15 @@ private:
 	/** Name of the client spec (workspace) passed in on the command line*/
 	FString ClientSpecName;
 
-	/** The number of the changelist being submitted, passed in on the command line */
-	FString ChangelistNumber;
+	/** The mode for the application to run */
+	EMode Mode = EMode::Unknown;
 
 	/** Bitfield control the various options that should be run */
 	EProcessOptions ProcessOptions = EProcessOptions::Virtualize;
+
+	/** The number of the changelist being submitted, used with EMode::Changelist  */
+	FString ChangelistNumber;
+
+	/** The path to the list of packages to virtualized, used with EMode::PackageList */
+	FString PackageListPath;
 };
