Wrap FUnrealVirtualizationToolApp in the UE::Virtualization namespace.

#rnx

- Some of our local structures, like FPLugin can get confused by the debugger with the version in the engine. Wrapping the tools code in the virtualization namespace helps keep the separation clear and stops the debugger from getting confused.
- Plus even though this is an UE program, I probably should've stuck to the new namespace rules anyway.

diff --git a/Engine/Source/Programs/UnrealVirtualizationTool/Private/UnrealVirtualizationTool.cpp b/Engine/Source/Programs/UnrealVirtualizationTool/Private/UnrealVirtualizationTool.cpp
--- a/Engine/Source/Programs/UnrealVirtualizationTool/Private/UnrealVirtualizationTool.cpp
+++ b/Engine/Source/Programs/UnrealVirtualizationTool/Private/UnrealVirtualizationTool.cpp
@@ -12,6 +12,8 @@ DEFINE_LOG_CATEGORY(LogVirtualizationTool);
 
 bool UnrealVirtualizationToolMain(int32 ArgC, TCHAR* ArgV[])
 {
+	using namespace UE::Virtualization;
+
 	GEngineLoop.PreInit(ArgC, ArgV);
 	check(GConfig && GConfig->IsReadyForUse());
 
diff --git a/Engine/Source/Programs/UnrealVirtualizationTool/Private/UnrealVirtualizationToolApp.cpp b/Engine/Source/Programs/UnrealVirtualizationTool/Private/UnrealVirtualizationToolApp.cpp
--- a/Engine/Source/Programs/UnrealVirtualizationTool/Private/UnrealVirtualizationToolApp.cpp
+++ b/Engine/Source/Programs/UnrealVirtualizationTool/Private/UnrealVirtualizationToolApp.cpp
@@ -19,6 +19,9 @@
 #include "UnrealVirtualizationTool.h"
 #include "Virtualization/VirtualizationSystem.h"
 
+namespace UE::Virtualization
+{
+
 namespace
 {
 
@@ -151,8 +154,6 @@ void LexFromString(EMode& OutValue, const FStringView& InString)
 	}
 }
 
-} // namespace
-
 /** 
  * This class can be used to prevent log messages from other systems being logged with the Display verbosity.
  * In practical terms this means as long as the class is alive, only LogVirtualizationTool messages will
@@ -186,6 +187,8 @@ private:
 	FFeedbackContext* OriginalLog;
 };
 
+} // namespace
+
 FUnrealVirtualizationToolApp::FUnrealVirtualizationToolApp()
 {
 
@@ -931,3 +934,5 @@ bool FProject::TryLoadConfig(FConfigFile& OutConfig) const
 
 	return true;
 }
+
+} // namespace UE::Virtualization
diff --git a/Engine/Source/Programs/UnrealVirtualizationTool/Private/UnrealVirtualizationToolApp.h b/Engine/Source/Programs/UnrealVirtualizationTool/Private/UnrealVirtualizationToolApp.h
--- a/Engine/Source/Programs/UnrealVirtualizationTool/Private/UnrealVirtualizationToolApp.h
+++ b/Engine/Source/Programs/UnrealVirtualizationTool/Private/UnrealVirtualizationToolApp.h
@@ -9,6 +9,9 @@
 
 class ISourceControlProvider;
 
+namespace UE::Virtualization
+{
+
 struct FPlugin
 {
 	FString PluginFilePath;
@@ -53,11 +56,11 @@ enum class EMode :uint32
 enum class EProcessOptions : uint32
 {
 	/** No options */
-	None		= 0,
+	None = 0,
 	/** Virtualize the packages in the provided changelist */
-	Virtualize	= 1 << 0,
+	Virtualize = 1 << 0,
 	/** Submit the changelist */
-	Submit		= 1 << 1
+	Submit = 1 << 1
 };
 ENUM_CLASS_FLAGS(EProcessOptions);
 
@@ -117,7 +120,7 @@ private:
 
 	/** Data structure holding the files in the changelist sorted by project and then by plugin*/
 	TArray<FProject> Projects;
-	
+
 	/** Name of the client spec (workspace) passed in on the command line*/
 	FString ClientSpecName;
 
@@ -133,3 +136,5 @@ private:
 	/** The path to the list of packages to virtualized, used with EMode::PackageList */
 	FString PackageListPath;
 };
+
+} //namespace UE::Virtualization
