Fix incorrect logging verbosity

#rnx

- The log message stating that the tool is now trying to submit a changelist was accidently set to 'Error' when it should be 'Display'

diff --git a/Engine/Source/Programs/UnrealVirtualizationTool/Private/UnrealVirtualizationToolApp.cpp b/Engine/Source/Programs/UnrealVirtualizationTool/Private/UnrealVirtualizationToolApp.cpp
--- a/Engine/Source/Programs/UnrealVirtualizationTool/Private/UnrealVirtualizationToolApp.cpp
+++ b/Engine/Source/Programs/UnrealVirtualizationTool/Private/UnrealVirtualizationToolApp.cpp
@@ -359,7 +359,7 @@ bool FUnrealVirtualizationToolApp::TrySubmitChangelist(const TArray<FString>& De
 		return true;
 	}
 
-	UE_LOG(LogVirtualizationTool, Error, TEXT("Attempting to submit the changelist '%s'"), *ChangelistNumber);
+	UE_LOG(LogVirtualizationTool, Display, TEXT("Attempting to submit the changelist '%s'"), *ChangelistNumber);
 
 	if (!SCCProvider.IsValid())
 	{
