UVT: Delay connecting to perforce until it is actually needed.

#rnx

- We now connect to perforce if either a) we need to parse a given changelist to find the package files b) we need to submit the given changelist
- This means we won't connect to perforce at all when running with the cmdline '-Mode=PackageList' saving a few seconds (although if the virtualization system is using the source control backend, then a connection will be made there anyway)

diff --git a/Engine/Source/Programs/UnrealVirtualizationTool/Private/UnrealVirtualizationToolApp.cpp b/Engine/Source/Programs/UnrealVirtualizationTool/Private/UnrealVirtualizationToolApp.cpp
--- a/Engine/Source/Programs/UnrealVirtualizationTool/Private/UnrealVirtualizationToolApp.cpp
+++ b/Engine/Source/Programs/UnrealVirtualizationTool/Private/UnrealVirtualizationToolApp.cpp
@@ -228,11 +228,6 @@ EInitResult FUnrealVirtualizationToolApp::Initialize()
 		return EInitResult::Error;
 	}
 
-	if (!TryConnectToSourceControl())
-	{
-		return EInitResult::Error;
-	}
-
 	TArray<FString> Packages;
 	switch (Mode)
 	{
@@ -276,12 +271,6 @@ bool FUnrealVirtualizationToolApp::Run()
 	{
 		UE_LOG(LogVirtualizationTool, Display, TEXT("Running the virtualization process..."));
 
-		if (!SCCProvider.IsValid())
-		{
-			UE_LOG(LogVirtualizationTool, Error, TEXT("No valid source control connection found!"));
-			return false;
-		}
-
 		TArray<FText> DescriptionTags;
 
 		for (const FProject& Project : Projects)
@@ -363,9 +352,11 @@ bool FUnrealVirtualizationToolApp::TrySubmitChangelist(const TArray<FString>& De
 
 	if (!SCCProvider.IsValid())
 	{
-		// This should not be possible, the check and error message is to guard against potential future problems only.
-		UE_LOG(LogVirtualizationTool, Error, TEXT("Submit failed, cannot find a valid source control provider"));
-		return false;
+		if (!TryConnectToSourceControl())
+		{
+			UE_LOG(LogVirtualizationTool, Error, TEXT("Submit failed, cannot find a valid source control provider"));
+			return false;
+		}
 	}
 
 	if (!ChangelistToSubmit.IsValid())
@@ -587,6 +578,11 @@ bool FUnrealVirtualizationToolApp::TryParseChangelist(TArray<FString>& OutPackag
 {
 	TRACE_CPUPROFILER_EVENT_SCOPE(TryParseChangelist);
 
+	if (!TryConnectToSourceControl())
+	{
+		return false;
+	}
+
 	UE_LOG(LogVirtualizationTool, Display, TEXT("Attempting to parse changelist '%s' in workspace '%s'"), *ChangelistNumber, *ClientSpecName);
 
 	if (!SCCProvider.IsValid())
