UVT: Clean up the logging when running in packagelist mode

#rnx

- We were printing out the packagelist file path both when it was found as a command line arg and when parsing it, we now only print it when it is found as an arg
- Any log section that indicates it could be slow has had "..." appended to it to show the user that this is expected
-- Any log message in a section like this is now tab indented to make the association easy for the user to see.

diff --git a/Engine/Source/Programs/UnrealVirtualizationTool/Private/UnrealVirtualizationToolApp.cpp b/Engine/Source/Programs/UnrealVirtualizationTool/Private/UnrealVirtualizationToolApp.cpp
--- a/Engine/Source/Programs/UnrealVirtualizationTool/Private/UnrealVirtualizationToolApp.cpp
+++ b/Engine/Source/Programs/UnrealVirtualizationTool/Private/UnrealVirtualizationToolApp.cpp
@@ -563,7 +563,7 @@ EInitResult FUnrealVirtualizationToolApp::TryParsePackageListCmdLine(const TCHAR
 {
 	if (FParse::Value(CmdLine, TEXT("-Path="), PackageListPath))
 	{
-		UE_LOG(LogVirtualizationTool, Display, TEXT("Attempting to virtualize package list '%s'"), *PackageListPath);
+		UE_LOG(LogVirtualizationTool, Display, TEXT("Virtualizing packages found in package list: '%s'"), *PackageListPath);
 		return EInitResult::Success;
 	}
 	else
@@ -666,11 +666,11 @@ bool FUnrealVirtualizationToolApp::TryParsePackageList(TArray<FString>& OutPacka
 {
 	TRACE_CPUPROFILER_EVENT_SCOPE(TrySortFilesByProject);
 
-	UE_LOG(LogVirtualizationTool, Display, TEXT("Attempting to parse package list '%s'"), *PackageListPath);
+	UE_LOG(LogVirtualizationTool, Display, TEXT("Parsing the package list..."));
 
 	if (!IFileManager::Get().FileExists(*PackageListPath))
 	{
-		UE_LOG(LogVirtualizationTool, Error, TEXT("The package list '%s' does not exist"), *PackageListPath);
+		UE_LOG(LogVirtualizationTool, Error, TEXT("\tThe package list '%s' does not exist"), *PackageListPath);
 		return false;
 	}
 
@@ -683,12 +683,12 @@ bool FUnrealVirtualizationToolApp::TryParsePackageList(TArray<FString>& OutPacka
 			FPaths::NormalizeFilename(PackagePath);
 		}
 
-		UE_LOG(LogVirtualizationTool, Display, TEXT("Found '%d' package file(s)"), OutPackages.Num());
+		UE_LOG(LogVirtualizationTool, Display, TEXT("\tFound '%d' package file(s)"), OutPackages.Num());
 		return true;
 	}
 	else
 	{
-		UE_LOG(LogVirtualizationTool, Error, TEXT("Failed to parse the package list '%s'"), *PackageListPath);
+		UE_LOG(LogVirtualizationTool, Error, TEXT("\tFailed to parse the package list '%s'"), *PackageListPath);
 		return false;
 	}
 }
@@ -697,7 +697,7 @@ bool FUnrealVirtualizationToolApp::TrySortFilesByProject(const TArray<FString>&
 {
 	TRACE_CPUPROFILER_EVENT_SCOPE(TrySortFilesByProject);
 
-	UE_LOG(LogVirtualizationTool, Display, TEXT("Sorting files by project"));
+	UE_LOG(LogVirtualizationTool, Display, TEXT("Sorting files by project..."));
 
 	for (const FString& PackagePath : Packages)
 	{
@@ -718,7 +718,7 @@ bool FUnrealVirtualizationToolApp::TrySortFilesByProject(const TArray<FString>&
 		}
 	}
 
-	UE_LOG(LogVirtualizationTool, Display, TEXT("The package files are associated with '%d' projects(s)"), Projects.Num());
+	UE_LOG(LogVirtualizationTool, Display, TEXT("\tThe package files are associated with '%d' projects(s)"), Projects.Num());
 
 	return true;
 }
