Fix the virtualization tool so that it will virtualize packages with more than one content directory in their path.

#rnx

- The main change is that if we do not find the project file relative to the first /content/ directory we will keep looking up the path to see if we can find additional /content/ directories.
- All log messages have been changed to warnings for consistency.

diff --git a/Engine/Source/Programs/UnrealVirtualizationTool/Private/UnrealVirtualizationToolApp.cpp b/Engine/Source/Programs/UnrealVirtualizationTool/Private/UnrealVirtualizationToolApp.cpp
--- a/Engine/Source/Programs/UnrealVirtualizationTool/Private/UnrealVirtualizationToolApp.cpp
+++ b/Engine/Source/Programs/UnrealVirtualizationTool/Private/UnrealVirtualizationToolApp.cpp
@@ -729,9 +729,18 @@ bool FUnrealVirtualizationToolApp::TryFindProject(const FString& PackagePath, FS
 
 	// TODO: This could be heavily optimized by caching known project files
 
-	const int32 ContentIndex = PackagePath.Find(TEXT("/content/"), ESearchCase::IgnoreCase, ESearchDir::FromEnd);
-	if (ContentIndex != INDEX_NONE)
+	int32 ContentIndex = PackagePath.Find(TEXT("/content/"), ESearchCase::IgnoreCase, ESearchDir::FromEnd);
+
+	// Early out if there is not a single content directory in the path
+	if (ContentIndex == INDEX_NONE)
+	{
+		UE_LOG(LogVirtualizationTool, Warning, TEXT("'%s' is not under a content directory"), *PackagePath);
+		return false;
+	}
+
+	while (ContentIndex != INDEX_NONE)
 	{
+		// Assume that the project directory is the parent of the /content/ directory
 		FString ProjectDirectory = PackagePath.Left(ContentIndex);
 		FString PluginDirectory;
 		
@@ -742,7 +751,7 @@ bool FUnrealVirtualizationToolApp::TryFindProject(const FString& PackagePath, FS
 
 		if (ProjectFile.IsEmpty())
 		{
-			// Could be a plugin so lets check
+			// If there was no project file, the package could be in a plugin, so lets check for that
 			PluginDirectory = ProjectDirectory;
 			IFileManager::Get().FindFiles(PluginFile, *PluginDirectory, TEXT(".uplugin"));
 
@@ -750,16 +759,18 @@ bool FUnrealVirtualizationToolApp::TryFindProject(const FString& PackagePath, FS
 			{
 				PluginFilePath = PluginDirectory / PluginFile[0];
 
+				// We have a valid plugin file, so we should be able to find a /plugins/ directory which will be just below the project directory
 				const int32 PluginIndex = PluginDirectory.Find(TEXT("/plugins/"), ESearchCase::IgnoreCase, ESearchDir::FromEnd);
 				if (PluginIndex != INDEX_NONE)
 				{
+					// We found the plugin root directory so the one above it should be the project directory
 					ProjectDirectory = PluginDirectory.Left(PluginIndex);
 					IFileManager::Get().FindFiles(ProjectFile, *ProjectDirectory, TEXT(".uproject"));
 				}
 			}
 			else if (PluginFile.Num() > 1)
 			{
-				UE_LOG(LogVirtualizationTool, Verbose, TEXT("Found multiple project files for '%s' at '%s'"), *PackagePath, *PluginDirectory);
+				UE_LOG(LogVirtualizationTool, Warning, TEXT("Found multiple .uplugin files for '%s' at '%s'"), *PackagePath, *PluginDirectory);
 				return false;
 			}
 		}
@@ -769,22 +780,19 @@ bool FUnrealVirtualizationToolApp::TryFindProject(const FString& PackagePath, FS
 			ProjectFilePath = ProjectDirectory / ProjectFile[0];
 			return true;
 		}
-		else if (ProjectFile.IsEmpty())
+		else if (!ProjectFile.IsEmpty())
 		{
-			UE_LOG(LogVirtualizationTool, Log, TEXT("Failed to find project file for '%s'"), *PackagePath);
+			UE_LOG(LogVirtualizationTool, Warning, TEXT("Found multiple .uproject files for '%s' at '%s'"), *PackagePath, *ProjectDirectory);
 			return false;
 		}
-		else
-		{
-			UE_LOG(LogVirtualizationTool, Verbose, TEXT("Found multiple project files for '%s' at '%s'"), *PackagePath, *ProjectDirectory);
-			return false;
-		}
-	}
-	else
-	{
-		UE_LOG(LogVirtualizationTool, Log, TEXT("File '%s' is not under a content directory"), *PackagePath);
-		return false;
+		
+		// Could be more than one content directory in the path so lets keep looking
+		ContentIndex = PackagePath.Find(TEXT("/content/"), ESearchCase::IgnoreCase, ESearchDir::FromEnd, ContentIndex);
 	}
+	
+	// We found one or more content directories but none of them contained a project file
+	UE_LOG(LogVirtualizationTool, Warning, TEXT("Failed to find project file for '%s'"), *PackagePath);
+	return false;
 }
 
 FProject& FUnrealVirtualizationToolApp::FindOrAddProject(const FString& ProjectFilePath)
