Move anon namespace out of the named namespace to avoid naming confusion.

#rnx

diff --git a/Engine/Source/Programs/UnrealVirtualizationTool/Private/UnrealVirtualizationToolApp.cpp b/Engine/Source/Programs/UnrealVirtualizationTool/Private/UnrealVirtualizationToolApp.cpp
--- a/Engine/Source/Programs/UnrealVirtualizationTool/Private/UnrealVirtualizationToolApp.cpp
+++ b/Engine/Source/Programs/UnrealVirtualizationTool/Private/UnrealVirtualizationToolApp.cpp
@@ -21,9 +21,6 @@
 #include "UnrealVirtualizationTool.h"
 #include "Virtualization/VirtualizationSystem.h"
 
-namespace UE::Virtualization
-{
-
 namespace
 {
 
@@ -37,9 +34,9 @@ bool IsPackageFile(const FString FilePath)
 	return FPackageName::IsPackageExtension(*Extension);
 }
 
-/** 
+/**
  * Utility to clean up the tags we got from the virtualization system. Convert the FText to FString and
- * discard any duplicate entries. 
+ * discard any duplicate entries.
  */
 TArray<FString> BuildFinalTagDescriptions(TArray<FText>& DescriptionTags)
 {
@@ -55,44 +52,44 @@ TArray<FString> BuildFinalTagDescriptions(TArray<FText>& DescriptionTags)
 }
 
 /** Utility to get EMode from a string */
-void LexFromString(EMode& OutValue, const FStringView& InString)
+void LexFromString(UE::Virtualization::EMode& OutValue, const FStringView& InString)
 {
 	if (InString == TEXT("Changelist"))
 	{
-		OutValue = EMode::Changelist;
+		OutValue = UE::Virtualization::EMode::Changelist;
 	}
 	else if (InString == TEXT("PackageList"))
 	{
-		OutValue = EMode::PackageList;
+		OutValue = UE::Virtualization::EMode::PackageList;
 	}
 	else if (InString == TEXT("Rehydrate"))
 	{
-		OutValue = EMode::Rehydrate;
+		OutValue = UE::Virtualization::EMode::Rehydrate;
 	}
 	else
 	{
-		OutValue = EMode::Unknown;
+		OutValue = UE::Virtualization::EMode::Unknown;
 	}
 }
 
 /** Utility for creating a new command */
 template<typename CommandType>
-TUniquePtr<FCommand> CreateCommand(const FString& ModeName, const TCHAR* CmdLine)
+TUniquePtr<UE::Virtualization::FCommand> CreateCommand(const FString& ModeName, const TCHAR* CmdLine)
 {
 	UE_LOG(LogVirtualizationTool, Display, TEXT("Attempting to initialize command '%s'..."), *ModeName);
 
-	TUniquePtr<FCommand> Command = MakeUnique<CommandType>(ModeName);
+	TUniquePtr<UE::Virtualization::FCommand> Command = MakeUnique<CommandType>(ModeName);
 	if (Command->Initialize(CmdLine))
 	{
 		return Command;
 	}
 	else
 	{
-		return TUniquePtr<FCommand>();
+		return TUniquePtr<UE::Virtualization::FCommand>();
 	}
 }
 
-/** 
+/**
  * This class can be used to prevent log messages from other systems being logged with the Display verbosity.
  * In practical terms this means as long as the class is alive, only LogVirtualizationTool messages will
  * be logged to the display meaning the user will have less information to deal with.
@@ -121,12 +118,15 @@ public:
 		FFeedbackContextAnsi::Serialize(V, Verbosity, Category);
 	}
 
-private: 
+private:
 	FFeedbackContext* OriginalLog;
 };
 
 } // namespace
 
+namespace UE::Virtualization
+{
+
 FUnrealVirtualizationToolApp::FUnrealVirtualizationToolApp()
 {
 
@@ -619,7 +619,7 @@ bool FUnrealVirtualizationToolApp::TryParseChangelist(TArray<FString>& OutPackag
 			const TArray<FSourceControlStateRef>& FilesinChangelist = ChangelistState->GetFilesStates();
 			for (const FSourceControlStateRef& FileState : FilesinChangelist)
 			{
-				if (IsPackageFile(FileState->GetFilename()))
+				if (::IsPackageFile(FileState->GetFilename()))
 				{
 					OutPackages.Add(FileState->GetFilename());
 				}
