UnrealVirtualizationTool can now be given the path of a package or the path of a directory of packages to virtualize in addition to previous functionality. NOTE that all existing functionality should continue to work but using -Mode=Changelist or -Mode=Packagelist will result in a warning prompting the future use of the new versions for those commands.

#rnx

- When the rehydration functionality was added it started to look like more functionality would be added to the tool than originally thought so I started to add the new functionality via a command system. This change now moves the older legacy functionality to the coommand system as well.
- Added a new FVirtualizeCommand which can accept package paths, directory paths, packagelist files or changelists as the input.
-- Unlike the legacy commands, virtualizing via changelist does not require the client spec to be provided on the commandline, although doing so will avoid several perforce commands and speed up the call. It is expected that people calling the tool on the commandline will probably opt to not supply a clientspec and let the tool workout which workspace a changelist is under, where as calls from other tools (such as P4VUtils) can provide it if already known to speed things up.
-- FVirtualizeLegacyChangeListCommand replicates the functionality of the old -Mode=Changelist command.
-- FVirtualizeLegacyPackageListCommand replicates the functionality of the old -Mode=Packagelist command.
-- The new command will only try to submit the results if a changelist was provided as the input and even then it will not do so by default. The tool now requires people to opt into submitting the changelist via the command line option -submit. NOTE: Other versions of the command maybe allow submission in the future but the virtualization process needs to be improved so that it can check out package files before this really makes sense.
- The rehydration command no longer requires a clientspec on the command line, it wasn't using the value anyway.
- Moved the source control code from the app code files to CommandBase. If we add more commands in the future we might want to factor this out to its own base class so commands can opt into source control functionality.
- In the future we should probably move the package -> project sorting code (FUnrealVirtualizationToolApp::TrySortFilesByProject) into the command base code as well.

diff --git a/Engine/Source/Programs/UnrealVirtualizationTool/Private/Commands/CommandBase.cpp b/Engine/Source/Programs/UnrealVirtualizationTool/Private/Commands/CommandBase.cpp
--- a/Engine/Source/Programs/UnrealVirtualizationTool/Private/Commands/CommandBase.cpp
+++ b/Engine/Source/Programs/UnrealVirtualizationTool/Private/Commands/CommandBase.cpp
@@ -2,6 +2,7 @@
 
 #include "CommandBase.h"
 
+#include "HAL/FileManager.h"
 #include "ISourceControlModule.h"
 #include "ISourceControlProvider.h"
 #include "Misc/PackageName.h"
@@ -24,6 +25,89 @@ bool IsPackageFile(const FString FilePath)
 	return FPackageName::IsPackageExtension(*Extension);
 }
 
+FCommand::FCommand(FStringView InCommandName)
+	: CommandName(InCommandName)
+{
+
+}
+
+FCommand::~FCommand()
+{
+
+}
+
+void FCommand::ParseCommandLine(const TCHAR* CmdLine, TArray<FString>& Tokens, TArray<FString>& Switches)
+{
+	// TODO: Taken from UCommandlet, maybe consider moving this code to a general purpose utility. We 
+	// could also make a version that returns TArray<FStringView>
+
+	FString NextToken;
+	while (FParse::Token(CmdLine, NextToken, false))
+	{
+		if (**NextToken == TCHAR('-'))
+		{
+			new(Switches) FString(NextToken.Mid(1));
+		}
+		else
+		{
+			new(Tokens) FString(NextToken);
+		}
+	}
+}
+
+FCommand::EPathResult FCommand::ParseSwitchForPaths(const FString& Switch, TArray<FString>& OutPackages)
+{
+	FString Path;
+	if (FParse::Value(*Switch, TEXT("Package="), Path))
+	{
+		FPaths::NormalizeFilename(Path);
+
+		if (!FPackageName::IsPackageFilename(Path))
+		{
+			UE_LOG(LogVirtualizationTool, Error, TEXT("Requested package file '%s' is not a valid package filename"), *Path);
+			return EPathResult::Error;
+		}
+
+		if (!IFileManager::Get().FileExists(*Path))
+		{
+			UE_LOG(LogVirtualizationTool, Error, TEXT("Could not find the requested package file '%s'"), *Path);
+			return EPathResult::Error;
+		}
+
+		OutPackages.Add(Path);
+		return EPathResult::Success;
+	}
+	else if (FParse::Value(*Switch, TEXT("PackageDir="), Path) || FParse::Value(*Switch, TEXT("PackageFolder="), Path))
+	{
+		// Note that 'PackageFolder' is the switch used by the resave commandlet, so allowing it here for compatibility purposes
+		FPaths::NormalizeFilename(Path);
+		if (IFileManager::Get().DirectoryExists(*Path))
+		{
+			IFileManager::Get().IterateDirectoryRecursively(*Path, [&OutPackages](const TCHAR* Path, bool bIsDirectory)
+				{
+					if (!bIsDirectory && FPackageName::IsPackageFilename(Path))
+					{
+						FString FilePath(Path);
+						FPaths::NormalizeFilename(FilePath);
+
+						OutPackages.Add(MoveTemp(FilePath));
+					}
+
+					return true; // Continue
+				});
+
+			return EPathResult::Success;
+		}
+		else
+		{
+			UE_LOG(LogVirtualizationTool, Error, TEXT("Could not find the requested directory '%s'"), *Path);
+			return EPathResult::Error;
+		}
+	}
+
+	return EPathResult::NotFound;
+}
+
 bool FCommand::TryConnectToSourceControl(FStringView ClientSpecName)
 {
 	if (SCCProvider.IsValid())
@@ -104,7 +188,7 @@ bool FCommand::TryCheckOutFilesForProject(FStringView ClientSpecName, FStringVie
 	return true;
 }
 
-bool FCommand::TryParseChangelist(FStringView ClientSpecName, FStringView ChangelistNumber, TArray<FString>& OutPackages)
+bool FCommand::TryParseChangelist(FStringView ClientSpecName, FStringView ChangelistNumber, TArray<FString>& OutPackages, FSourceControlChangelistPtr* OutChangelist)
 {
 	TRACE_CPUPROFILER_EVENT_SCOPE(TryParseChangelist);
 
@@ -180,6 +264,11 @@ bool FCommand::TryParseChangelist(FStringView ClientSpecName, FStringView Change
 				}
 			}
 
+			if (OutChangelist != nullptr)
+			{
+				*OutChangelist = Changelist;
+			}
+
 			return true;
 		}
 	}
@@ -189,6 +278,50 @@ bool FCommand::TryParseChangelist(FStringView ClientSpecName, FStringView Change
 	return false;
 }
 
+FString FCommand::FindClientSpecForChangelist(FStringView ChangelistNumber)
+{
+	TRACE_CPUPROFILER_EVENT_SCOPE(FindClientSpecForChangelist);
+
+	UE_LOG(LogVirtualizationTool, Display, TEXT("\tAttempting to find the workspace for '%.*s'"),
+		ChangelistNumber.Len(), ChangelistNumber.GetData());
+
+	if (!TryConnectToSourceControl(FStringView()))
+	{
+		return FString();
+	}
+
+	if (!SCCProvider.IsValid())
+	{
+		UE_LOG(LogVirtualizationTool, Error, TEXT("No valid source control connection found!"));
+		return FString();
+	}
+
+	if (!SCCProvider->UsesChangelists())
+	{
+		UE_LOG(LogVirtualizationTool, Error, TEXT("The source control provider does not support the use of changelists"));
+		return FString();
+	}
 
+	TSharedRef<FGetChangelistDetails> Operation = ISourceControlOperation::Create<FGetChangelistDetails>(ChangelistNumber);
+
+	if (SCCProvider->Execute(Operation, EConcurrency::Synchronous) == ECommandResult::Succeeded && !Operation->GetChangelistDetails().IsEmpty())
+	{
+		const FString* ClientSpec = Operation->GetChangelistDetails()[0].Find(TEXT("Client"));
+		if (ClientSpec != nullptr)
+		{
+			return *ClientSpec;
+		}
+		else
+		{ 
+			UE_LOG(LogVirtualizationTool, Error, TEXT("Unable to find the 'client' field for the changelist '%.*s'"), ChangelistNumber.Len(), ChangelistNumber.GetData());
+			return FString();
+		}
+	}
+	else
+	{
+		UE_LOG(LogVirtualizationTool, Error, TEXT("Failed to find details for the changelist '%.*s'"), ChangelistNumber.Len(), ChangelistNumber.GetData());
+		return FString();
+	}
+}
 
-} //namespace UE::Virtualization 
\ No newline at end of file
+} //namespace UE::Virtualization 
diff --git a/Engine/Source/Programs/UnrealVirtualizationTool/Private/Commands/CommandBase.h b/Engine/Source/Programs/UnrealVirtualizationTool/Private/Commands/CommandBase.h
--- a/Engine/Source/Programs/UnrealVirtualizationTool/Private/Commands/CommandBase.h
+++ b/Engine/Source/Programs/UnrealVirtualizationTool/Private/Commands/CommandBase.h
@@ -5,9 +5,13 @@
 #include "Containers/Array.h"
 #include "Containers/StringView.h"
 #include "Containers/UnrealString.h"
+#include "Templates/SharedPointer.h"
 #include "Templates/UniquePtr.h"
 
 class ISourceControlProvider;
+class ISourceControlChangelist;
+
+typedef TSharedPtr<ISourceControlChangelist, ESPMode::ThreadSafe> FSourceControlChangelistPtr;
 
 namespace UE::Virtualization
 {
@@ -18,29 +22,37 @@ class FProject;
 class FCommand
 {
 public:
-	FCommand(FStringView InCommandName)
-		: CommandName(InCommandName)
-	{
-
-	}
-	virtual ~FCommand() = default;
+	FCommand(FStringView InCommandName);
+	virtual ~FCommand();
 
 	virtual bool Initialize(const TCHAR* CmdLine) = 0;
 
 	virtual bool Run(const TArray<FProject>& Projects) = 0;
 
+	virtual const TArray<FString>& GetPackages() const = 0;
+
 	const FString& GetName() const
 	{
 		return CommandName;
 	}
 
-	const TArray<FString>& GetPackages() const
+protected: // Common commandline parsing code
+	
+	static void ParseCommandLine(const TCHAR* CmdLine, TArray<FString>& Tokens, TArray<FString>& Switches);
+
+	enum EPathResult : int8
 	{
-		return Packages;
-	}
+		/** The switch was a valid package/path but parsing it resulted in an error */
+		Error = -1,
+		/** The switch was not a valid package/path */
+		NotFound = 0,
+		/** The switch was a valid package/path and was successfully parsed */
+		Success = 1
+	};
 
-protected:
+	static EPathResult ParseSwitchForPaths(const FString& Switch, TArray<FString>& OutPackages);
 
+protected: // Common SourceControl Code
 	/**
 	 * Creates a new ISourceControlProvider that can be used to perforce source control operations. This provider
 	 * will remain valid until reset or the command is completed.
@@ -53,15 +65,14 @@ protected:
 
 	bool TryCheckOutFilesForProject(FStringView ClientSpecName, FStringView ProjectRoot, const TArray<FString>& ProjectPackages);
 	
-	bool TryParseChangelist(FStringView ClientSpecName, FStringView ChangelistNumber, TArray<FString>& OutPackages);
+	bool TryParseChangelist(FStringView ClientSpecName, FStringView ChangelistNumber, TArray<FString>& OutPackages, FSourceControlChangelistPtr* OutChangelist);
+
+	FString FindClientSpecForChangelist(FStringView ChangelistNumber);
 private:
 	FString CommandName;
 
 protected:
 	TUniquePtr<ISourceControlProvider> SCCProvider;
-
-	TArray<FString> Packages;
-
 };
 
 } // namespace UE::Virtualization
\ No newline at end of file
diff --git a/Engine/Source/Programs/UnrealVirtualizationTool/Private/Commands/RehydrateCommand.cpp b/Engine/Source/Programs/UnrealVirtualizationTool/Private/Commands/RehydrateCommand.cpp
--- a/Engine/Source/Programs/UnrealVirtualizationTool/Private/Commands/RehydrateCommand.cpp
+++ b/Engine/Source/Programs/UnrealVirtualizationTool/Private/Commands/RehydrateCommand.cpp
@@ -2,7 +2,6 @@
 
 #include "Commands/RehydrateCommand.h"
 
-#include "HAL/FileManager.h"
 #include "ISourceControlProvider.h"
 #include "Misc/ConfigCacheIni.h"
 #include "Misc/PackageName.h"
@@ -15,48 +14,24 @@
 namespace UE::Virtualization
 {
 
-// Taken from UCommandlet
-static void ParseCommandLine(const TCHAR* CmdLine, TArray<FString>& Tokens, TArray<FString>& Switches)
-{
-	FString NextToken;
-	while (FParse::Token(CmdLine, NextToken, false))
-	{
-		if (**NextToken == TCHAR('-'))
-		{
-			new(Switches) FString(NextToken.Mid(1));
-		}
-		else
-		{
-			new(Tokens) FString(NextToken);
-		}
-	}
-}
-
 FRehydrateCommand::FRehydrateCommand(FStringView CommandName)
 	: FCommand(CommandName)
 {
 
 }
 
-FRehydrateCommand::~FRehydrateCommand()
-{
-
-}
-
 void FRehydrateCommand::PrintCmdLineHelp()
 {
-	UE_LOG(LogVirtualizationTool, Display, TEXT("UnrealVirtualizationTool -ClientSpecName=<name> [optional] -Mode=Rehydrate -Package=<string> [global options]"));
-	UE_LOG(LogVirtualizationTool, Display, TEXT("\tRehyrates the given package"));
-
-	UE_LOG(LogVirtualizationTool, Display, TEXT("UnrealVirtualizationTool -ClientSpecName=<name> [optional] -Mode=Rehydrate -PackageDir=<string> [global options]"));
-	UE_LOG(LogVirtualizationTool, Display, TEXT("\tRehyrates all packages in the given directory and its subdirectories"));
-
-	UE_LOG(LogVirtualizationTool, Display, TEXT("UnrealVirtualizationTool -ClientSpecName=<name> [optional] -Mode=Rehydrate -Changelist=<number> [global options]"));
-	UE_LOG(LogVirtualizationTool, Display, TEXT("\tRehyrates all packages in the given changelist"));
+	UE_LOG(LogVirtualizationTool, Display, TEXT(""));
+	UE_LOG(LogVirtualizationTool, Display, TEXT("-Mode=Rehydrate -Package=<string>"));
+	UE_LOG(LogVirtualizationTool, Display, TEXT("-Mode=Rehydrate -PackageDir=<string>"));
+	UE_LOG(LogVirtualizationTool, Display, TEXT("-Mode=Rehydrate -Changelist=<number>"));
 }
 
 bool FRehydrateCommand::Initialize(const TCHAR* CmdLine)
 {
+	TRACE_CPUPROFILER_EVENT_SCOPE(FRehydrateCommand::Initialize);
+
 	// Note that we haven't loaded any projects config files and so don't really have
 	// any valid project mount points so we cannot use FPackagePath or FPackageName
 	// and expect to find anything!
@@ -65,62 +40,28 @@ bool FRehydrateCommand::Initialize(const TCHAR* CmdLine)
 	TArray<FString> Switches;
 
 	ParseCommandLine(CmdLine, Tokens, Switches);
-
-	Packages.Reserve(32);
-
+	
+	FString SwitchValue;
 	for (const FString& Switch :  Switches)
 	{
-		FString Path;
-		FString ChangelistNumber;
-
-		if (FParse::Value(*Switch, TEXT("ClientSpecName="), ClientSpecName))
+		EPathResult Result = ParseSwitchForPaths(Switch, Packages);
+		if (Result == EPathResult::Error)
 		{
-			UE_LOG(LogVirtualizationTool, Display, TEXT("\tWorkspace name provided '%s'"), *ClientSpecName);
+			return false;
 		}
-		else if (FParse::Value(*Switch, TEXT("Package="), Path))
+		else if (Result == EPathResult::Success)
 		{
-			FPaths::NormalizeFilename(Path);
-
-			if (IFileManager::Get().FileExists(*Path))
-			{
-				Packages.Add(Path);
-			}
-			else
-			{
-				UE_LOG(LogVirtualizationTool, Error, TEXT("Could not find the requested package file '%s'"), *Path);
-				return false;
-			}
+			continue; // If we already matched the switch we don't need to check against any others
 		}
-		else if (FParse::Value(*Switch, TEXT("PackageDir="), Path) || FParse::Value(*Switch, TEXT("PackageFolder="), Path))
-		{
-			// Note that 'PackageFolder' is the switch used by the resave commandlet, so allowing it here for compatibility purposes
-			FPaths::NormalizeFilename(Path);
-			if (IFileManager::Get().DirectoryExists(*Path))
-			{
-				IFileManager::Get().IterateDirectoryRecursively(*Path, [this](const TCHAR* Path, bool bIsDirectory)
-					{
-						if (!bIsDirectory && FPackageName::IsPackageFilename(Path))
-						{
-							FString FilePath(Path);
-							FPaths::NormalizeFilename(FilePath);
-
-							this->Packages.Add(MoveTemp(FilePath));
-						}
-
-						return true; // Continue
-					});
 
-				return true;
-			}
-			else
-			{
-				UE_LOG(LogVirtualizationTool, Error, TEXT("Could not find the requested directory '%s'"), *Path);
-				return false;
-			}
+		if (FParse::Value(*Switch, TEXT("ClientSpecName="), SwitchValue))
+		{
+			ClientSpecName = SwitchValue;
+			UE_LOG(LogVirtualizationTool, Display, TEXT("\tWorkspace name provided '%s'"), *ClientSpecName);
 		}
-		else if (FParse::Value(*Switch, TEXT("Changelist="), ChangelistNumber))
+		else if (FParse::Value(*Switch, TEXT("Changelist="), SwitchValue))
 		{
-			return TryParseChangelist(ClientSpecName, ChangelistNumber, Packages);
+			return TryParseChangelist(ClientSpecName, SwitchValue, Packages, nullptr);
 		}
 	}
 
@@ -129,6 +70,8 @@ bool FRehydrateCommand::Initialize(const TCHAR* CmdLine)
 
 bool FRehydrateCommand::Run(const TArray<FProject>& Projects)
 {
+	TRACE_CPUPROFILER_EVENT_SCOPE(FRehydrateCommand::Run);
+
 	for (const FProject& Project : Projects)
 	{
 		TStringBuilder<128> ProjectName;
@@ -182,4 +125,9 @@ bool FRehydrateCommand::Run(const TArray<FProject>& Projects)
 	return true;
 }
 
+const TArray<FString>& FRehydrateCommand::GetPackages() const
+{
+	return Packages;
+}
+
 } // namespace UE::Virtualization
diff --git a/Engine/Source/Programs/UnrealVirtualizationTool/Private/Commands/RehydrateCommand.h b/Engine/Source/Programs/UnrealVirtualizationTool/Private/Commands/RehydrateCommand.h
--- a/Engine/Source/Programs/UnrealVirtualizationTool/Private/Commands/RehydrateCommand.h
+++ b/Engine/Source/Programs/UnrealVirtualizationTool/Private/Commands/RehydrateCommand.h
@@ -11,7 +11,7 @@ class FRehydrateCommand final : public FCommand
 {
 public:
 	FRehydrateCommand(FStringView CommandName);
-	virtual ~FRehydrateCommand();
+	virtual ~FRehydrateCommand() = default;
 
 	static void PrintCmdLineHelp();
 private:
@@ -19,9 +19,13 @@ private:
 
 	virtual bool Run(const TArray<FProject>& Projects) override;
 
+	virtual const TArray<FString>& GetPackages() const override;
+
 private:
 
 	FString ClientSpecName;
+
+	TArray<FString> Packages;
 };
 
 } // namespace UE::Virtualization
\ No newline at end of file
diff --git a/Engine/Source/Programs/UnrealVirtualizationTool/Private/Commands/VirtualizeCommand.cpp b/Engine/Source/Programs/UnrealVirtualizationTool/Private/Commands/VirtualizeCommand.cpp
new file mode 100644
--- /dev/null
+++ b/Engine/Source/Programs/UnrealVirtualizationTool/Private/Commands/VirtualizeCommand.cpp
@@ -0,0 +1,342 @@
+// Copyright Epic Games, Inc. All Rights Reserved.
+
+#include "Commands/VirtualizeCommand.h"
+
+#include "HAL/FileManager.h"
+#include "Misc/FileHelper.h"
+#include "Misc/ScopeExit.h"
+#include "ProjectFiles.h"
+#include "UnrealVirtualizationTool.h"
+#include "Virtualization/VirtualizationSystem.h"
+
+
+#include "ISourceControlModule.h"
+#include "ISourceControlProvider.h"
+#include "SourceControlOperations.h"
+
+namespace
+{
+
+/** Utility */
+TArray<FString> BuildFinalTagDescriptions(TArray<FText>& DescriptionTags)
+{
+	TArray<FString> CleanedDescriptions;
+	CleanedDescriptions.Reserve(DescriptionTags.Num());
+
+	for (const FText& Tag : DescriptionTags)
+	{
+		CleanedDescriptions.AddUnique(Tag.ToString());
+	}
+
+	return CleanedDescriptions;
+}
+
+} // namespace
+
+namespace UE::Virtualization
+{
+
+FVirtualizeCommand::FVirtualizeCommand(FStringView CommandName)
+	: FCommand(CommandName)
+{
+
+}
+
+void FVirtualizeCommand::PrintCmdLineHelp()
+{
+	UE_LOG(LogVirtualizationTool, Display, TEXT(""));
+	UE_LOG(LogVirtualizationTool, Display, TEXT("-Mode=Virtualize -Changelist=<number> -Submit [optional]"));
+	UE_LOG(LogVirtualizationTool, Display, TEXT("-Mode=Virtualize -Path=<string>"));
+}
+
+bool FVirtualizeCommand::Initialize(const TCHAR* CmdLine)
+{
+	TRACE_CPUPROFILER_EVENT_SCOPE(FVirtualizeCommand::Initialize);
+
+	TArray<FString> Tokens;
+	TArray<FString> Switches;
+
+	ParseCommandLine(CmdLine, Tokens, Switches);
+
+	FString SwitchValue;
+	for (const FString& Switch : Switches)
+	{
+		EPathResult Result = ParseSwitchForPaths(Switch, AllPackages);
+		if (Result == EPathResult::Error)
+		{
+			return false;
+		}
+		else if (Result == EPathResult::Success)
+		{
+			continue; // If we already matched the switch we don't need to check against any others
+		}
+		
+		if (FParse::Value(*Switch, TEXT("Changelist="), SwitchValue))
+		{
+			SourceChangelistNumber = SwitchValue;
+		}
+		else if (Switch == TEXT("Submit"))
+		{
+			bShouldSubmitChangelist = true;
+		}
+	}
+
+	// Process the provided changelist if one was found
+	if (!SourceChangelistNumber.IsEmpty() )
+	{
+		// If no client spec was provided we need to find it for the changelist
+		// In theory this duplicates a lot of the work found in ::TryParseChangelist
+		// but at the moment the FGetChangelistDetails operation is not compatible
+		// with the FSourceControlChangelistStateRef/FSourceControlStateRef API
+		// so we are stuck with duplication of work.
+		if (ClientSpecName.IsEmpty())
+		{
+			ClientSpecName = FindClientSpecForChangelist(SourceChangelistNumber);
+			if (!ClientSpecName.IsEmpty())
+			{
+				FSourceControlResultInfo Info;
+				if (SCCProvider->SwitchWorkspace(ClientSpecName, Info, nullptr) != ECommandResult::Succeeded)
+				{
+					UE_LOG(LogVirtualization, Error, TEXT("Failed to switch to workspace '%s'"), *ClientSpecName);
+					return false;
+				}
+			}
+			else
+			{
+				UE_LOG(LogVirtualization, Error, TEXT("Count not find a valid workspace for the changelist '%s'"), *SourceChangelistNumber);
+				return false;
+			}
+		}
+
+		if (!TryParseChangelist(ClientSpecName, SourceChangelistNumber, AllPackages, &SourceChangelist))
+		{
+			UE_LOG(LogVirtualization, Error, TEXT("Failed to find the files in the changelist '%s'"), *SourceChangelistNumber);
+			return false;
+		}
+	}
+
+	return true;
+}
+
+bool FVirtualizeCommand::Run(const TArray<FProject>& Projects)
+{
+	TRACE_CPUPROFILER_EVENT_SCOPE(FVirtualizeCommand::Run);
+
+	UE_LOG(LogVirtualizationTool, Display, TEXT("Running the virtualization process..."));
+	
+	TArray<FText> DescriptionTags;
+
+	for (const FProject& Project : Projects)
+	{
+		TStringBuilder<128> ProjectName;
+		ProjectName << Project.GetProjectName();
+
+		UE_LOG(LogVirtualizationTool, Display, TEXT("\tChecking package(s) for the project '%s'..."), ProjectName.ToString());
+
+		FConfigFile EngineConfigWithProject;
+		if (!Project.TryLoadConfig(EngineConfigWithProject))
+		{
+			return false;
+		}
+
+		Project.RegisterMountPoints();
+
+		UE::Virtualization::FInitParams InitParams(ProjectName, EngineConfigWithProject);
+		UE::Virtualization::Initialize(InitParams, UE::Virtualization::EInitializationFlags::ForceInitialize);
+
+		TArray<FString> ProjectPackages = Project.GetAllPackages();
+
+		TArray<FText> Errors;
+		UE::Virtualization::IVirtualizationSystem::Get().TryVirtualizePackages(ProjectPackages, DescriptionTags, Errors);
+
+		if (!Errors.IsEmpty())
+		{
+			UE_LOG(LogVirtualizationTool, Error, TEXT("The virtualization process failed with the following errors:"));
+			for (const FText& Error : Errors)
+			{
+				UE_LOG(LogVirtualizationTool, Error, TEXT("\t%s"), *Error.ToString());
+			}
+			return false;
+		}
+
+		UE_LOG(LogVirtualizationTool, Display, TEXT("\tCheck complete"));
+
+		UE::Virtualization::Shutdown();
+		Project.UnRegisterMountPoints();
+	}
+
+	if (bShouldSubmitChangelist)
+	{
+		TArray<FString> FinalDescriptionTags = BuildFinalTagDescriptions(DescriptionTags);
+
+		if (!TrySubmitChangelist(SourceChangelist, FinalDescriptionTags))
+		{
+			return false;
+		}
+	}
+
+	return true;
+}
+
+const TArray<FString>& FVirtualizeCommand::GetPackages() const
+{
+	return AllPackages;
+}
+
+bool FVirtualizeCommand::TrySubmitChangelist(const FSourceControlChangelistPtr& ChangelistToSubmit, const TArray<FString>& DescriptionTags)
+{
+	if (!ChangelistToSubmit.IsValid())
+	{
+		return true;
+	}
+
+	if (!TryConnectToSourceControl(ClientSpecName))
+	{
+		UE_LOG(LogVirtualizationTool, Error, TEXT("Submit failed, cannot find a valid source control provider"));
+		return false;
+	}
+
+	TRACE_CPUPROFILER_EVENT_SCOPE(FVirtualizeCommand::TrySubmitChangelist);
+
+	FSourceControlChangelistRef Changelist = ChangelistToSubmit.ToSharedRef();
+	FSourceControlChangelistStatePtr ChangelistState = SCCProvider->GetState(Changelist, EStateCacheUsage::Use);
+
+	if (!ChangelistState.IsValid())
+	{
+		UE_LOG(LogVirtualizationTool, Error, TEXT("Submit failed, failed to find the state for the changelist"));
+		return false;
+	}
+
+	const FString ChangelistNumber = ChangelistState->GetDisplayText().ToString();
+
+	UE_LOG(LogVirtualizationTool, Display, TEXT("Attempting to submit the changelist '%s'"), *ChangelistNumber);
+
+	TSharedRef<FCheckIn, ESPMode::ThreadSafe> CheckInOperation = ISourceControlOperation::Create<FCheckIn>();
+
+	// Grab the original changelist description then append our tags afterwards
+	TStringBuilder<512> Description;
+	Description << ChangelistState->GetDescriptionText().ToString();
+
+	for (const FString& Tag : DescriptionTags)
+	{
+		Description << TEXT("\n") << Tag;
+	}
+
+	CheckInOperation->SetDescription(FText::FromString(Description.ToString()));
+
+	if (SCCProvider->Execute(CheckInOperation, Changelist) == ECommandResult::Succeeded)
+	{
+		UE_LOG(LogVirtualizationTool, Display, TEXT("%s"), *CheckInOperation->GetSuccessMessage().ToString());
+		return true;
+	}
+	else
+	{
+		// Even when log suppression is active we still show errors to the users and as the source control
+		// operation should have logged the problem as an error the user will see it. This means we don't 
+		// have to extract it from CheckInOperation .
+		UE_LOG(LogVirtualizationTool, Error, TEXT("Submit failed, please check the log!"));
+
+		return false;
+	}
+}
+
+bool FVirtualizeCommand::TryParsePackageList(const FString& PackageListPath, TArray<FString>& OutPackages)
+{
+	TRACE_CPUPROFILER_EVENT_SCOPE(FVirtualizeCommand::TryParsePackageList);
+
+	UE_LOG(LogVirtualizationTool, Display, TEXT("Parsing the package list '%s'..."), *PackageListPath);
+
+	if (!IFileManager::Get().FileExists(*PackageListPath))
+	{
+		UE_LOG(LogVirtualizationTool, Error, TEXT("\tThe package list '%s' does not exist"), *PackageListPath);
+		return false;
+	}
+
+	if (FFileHelper::LoadFileToStringArray(OutPackages, *PackageListPath))
+	{
+		// We don't have control over how the package list was generated so make sure that the paths
+		// are in the format that we want.
+		for (FString& PackagePath : OutPackages)
+		{
+			FPaths::NormalizeFilename(PackagePath);
+		}
+
+		return true;
+	}
+	else
+	{
+		UE_LOG(LogVirtualizationTool, Error, TEXT("\tFailed to parse the package list '%s'"), *PackageListPath);
+		return false;
+	}
+}
+
+
+FVirtualizeLegacyChangeListCommand::FVirtualizeLegacyChangeListCommand(FStringView CommandName)
+	: FVirtualizeCommand(CommandName)
+{
+
+}
+
+bool FVirtualizeLegacyChangeListCommand::Initialize(const TCHAR* CmdLine)
+{
+	TRACE_CPUPROFILER_EVENT_SCOPE(FVirtualizeLegacyChangeListCommand::Initialize);
+
+	UE_LOG(LogVirtualizationTool, Warning, TEXT("Using legacy -Mode=Changelist command, use '-Mode=Virtualization -Changelist=123' instead!"));
+
+	if (!FParse::Value(CmdLine, TEXT("-ClientSpecName="), ClientSpecName))
+	{
+		UE_LOG(LogVirtualizationTool, Error, TEXT("Failed to find cmdline switch 'ClientSpecName', this is a required parameter!"));
+		return false;
+	}
+
+	if (!FParse::Value(CmdLine, TEXT("-Changelist="), SourceChangelistNumber))
+	{
+		UE_LOG(LogVirtualizationTool, Error, TEXT("Failed to find cmdline switch 'Changelist', this is a required parameter!"));
+		return false;
+
+	}
+
+	if (!TryParseChangelist(ClientSpecName, SourceChangelistNumber, AllPackages, &SourceChangelist))
+	{
+		return false;
+	}
+
+	if (FParse::Param(CmdLine, TEXT("NoSubmit")))
+	{
+		UE_LOG(LogVirtualizationTool, Display, TEXT("Cmdline parameter '-NoSubmit' found, the changelist will be virtualized but not submitted!"));
+		bShouldSubmitChangelist = false;
+	}
+	else
+	{
+		// The legacy command was opt out when it came to submitting the changelist, we need to maintain those defaults
+		bShouldSubmitChangelist = true;
+	}
+
+	return true;
+}
+
+FVirtualizeLegacyPackageListCommand::FVirtualizeLegacyPackageListCommand(FStringView CommandName)
+	: FVirtualizeCommand(CommandName)
+{
+
+}
+
+bool FVirtualizeLegacyPackageListCommand::Initialize(const TCHAR* CmdLine)
+{
+	TRACE_CPUPROFILER_EVENT_SCOPE(FVirtualizeLegacyPackageListCommand::Initialize);
+
+	UE_LOG(LogVirtualizationTool, Warning, TEXT("Using legacy -Mode=Packagelist command, use '-Mode=Virtualization -Path=PathToFile' instead!"));
+
+	FString PackageListPath;
+	if (FParse::Value(CmdLine, TEXT("-Path="), PackageListPath))
+	{
+		return TryParsePackageList(PackageListPath, AllPackages);
+	}
+	else
+	{
+
+		return false;
+	}
+}
+
+} // namespace UE::Virtualization
diff --git a/Engine/Source/Programs/UnrealVirtualizationTool/Private/Commands/VirtualizeCommand.h b/Engine/Source/Programs/UnrealVirtualizationTool/Private/Commands/VirtualizeCommand.h
new file mode 100644
--- /dev/null
+++ b/Engine/Source/Programs/UnrealVirtualizationTool/Private/Commands/VirtualizeCommand.h
@@ -0,0 +1,66 @@
+// Copyright Epic Games, Inc. All Rights Reserved.
+
+#pragma once
+
+#include "Commands/CommandBase.h"
+
+namespace UE::Virtualization
+{
+
+class FVirtualizeCommand : public FCommand
+{
+public:
+	FVirtualizeCommand(FStringView CommandName);
+	virtual ~FVirtualizeCommand() = default;
+
+	static void PrintCmdLineHelp();
+
+private:
+	virtual bool Initialize(const TCHAR* CmdLine) override;
+
+	virtual bool Run(const TArray<FProject>& Projects) override;
+
+	virtual const TArray<FString>& GetPackages() const override;
+
+	bool TrySubmitChangelist(const FSourceControlChangelistPtr& ChangelistToSubmit, const TArray<FString>& DescriptionTags);
+	
+protected:
+
+	bool TryParsePackageList(const FString& PackageListPath, TArray<FString>& OutPackages);	
+
+protected:
+	FString ClientSpecName;
+
+	TArray<FString> AllPackages;
+
+	FSourceControlChangelistPtr SourceChangelist;
+	FString SourceChangelistNumber;
+
+	bool bShouldSubmitChangelist = false;
+};
+
+/** Used to convert the legacy command line -Mode=Changelist to the new style of command */
+class FVirtualizeLegacyChangeListCommand final : public FVirtualizeCommand
+{
+public:
+	FVirtualizeLegacyChangeListCommand(FStringView CommandName);
+	virtual ~FVirtualizeLegacyChangeListCommand() = default;
+
+private:
+	virtual bool Initialize(const TCHAR* CmdLine) override;
+};
+
+/** Used to convert the legacy command line -Mode=Packagelist to the new style of command */
+class FVirtualizeLegacyPackageListCommand final : public FVirtualizeCommand
+{
+public:
+	FVirtualizeLegacyPackageListCommand(FStringView CommandName);
+	virtual ~FVirtualizeLegacyPackageListCommand() = default;
+
+private:
+	virtual bool Initialize(const TCHAR* CmdLine) override;
+};
+
+
+
+} // namespace UE::Virtualization
diff --git a/Engine/Source/Programs/UnrealVirtualizationTool/Private/UnrealVirtualizationTool.cpp b/Engine/Source/Programs/UnrealVirtualizationTool/Private/UnrealVirtualizationTool.cpp
--- a/Engine/Source/Programs/UnrealVirtualizationTool/Private/UnrealVirtualizationTool.cpp
+++ b/Engine/Source/Programs/UnrealVirtualizationTool/Private/UnrealVirtualizationTool.cpp
@@ -17,6 +17,15 @@ bool UnrealVirtualizationToolMain(int32 ArgC, TCHAR* ArgV[])
 	GEngineLoop.PreInit(ArgC, ArgV);
 	check(GConfig && GConfig->IsReadyForUse());
 
+#if 0
+	while (!FPlatformMisc::IsDebuggerPresent())
+	{
+		FPlatformProcess::SleepNoStats(0.0f);
+	}
+
+	PLATFORM_BREAK();
+#endif
+
 	FModuleManager::Get().StartProcessingNewlyLoadedObjects();
 
 	bool bRanSuccessfully = true;
diff --git a/Engine/Source/Programs/UnrealVirtualizationTool/Private/UnrealVirtualizationToolApp.cpp b/Engine/Source/Programs/UnrealVirtualizationTool/Private/UnrealVirtualizationToolApp.cpp
--- a/Engine/Source/Programs/UnrealVirtualizationTool/Private/UnrealVirtualizationToolApp.cpp
+++ b/Engine/Source/Programs/UnrealVirtualizationTool/Private/UnrealVirtualizationToolApp.cpp
@@ -4,12 +4,11 @@
 
 #include "Commands/CommandBase.h"
 #include "Commands/RehydrateCommand.h"
+#include "Commands/VirtualizeCommand.h"
 #include "GenericPlatform/GenericPlatformOutputDevices.h"
 #include "HAL/FeedbackContextAnsi.h"
 #include "HAL/FileManager.h"
 #include "Interfaces/IPluginManager.h"
-#include "ISourceControlModule.h"
-#include "ISourceControlProvider.h"
 #include "Misc/CommandLine.h"
 #include "Misc/ConfigCacheIni.h"
 #include "Misc/FileHelper.h"
@@ -17,41 +16,13 @@
 #include "Misc/Parse.h"
 #include "Misc/Paths.h"
 #include "Misc/PathViews.h"
-#include "SourceControlInitSettings.h"
-#include "SourceControlOperations.h"
+#include "Modules/ModuleManager.h"
 #include "UnrealVirtualizationTool.h"
 #include "Virtualization/VirtualizationSystem.h"
 
 namespace
 {
 
-/** Utility for testing if a file path resolves to a valid package file or not */
-bool IsPackageFile(const FString FilePath)
-{
-	// ::IsPackageExtension requires a TCHAR so we cannot use FPathViews here
-	const FString Extension = FPaths::GetExtension(FilePath);
-
-	// Currently we don't virtualize text based assets so no call to FPackageName::IsTextPackageExtension
-	return FPackageName::IsPackageExtension(*Extension);
-}
-
-/**
- * Utility to clean up the tags we got from the virtualization system. Convert the FText to FString and
- * discard any duplicate entries.
- */
-TArray<FString> BuildFinalTagDescriptions(TArray<FText>& DescriptionTags)
-{
-	TArray<FString> CleanedDescriptions;
-	CleanedDescriptions.Reserve(DescriptionTags.Num());
-
-	for (const FText& Tag : DescriptionTags)
-	{
-		CleanedDescriptions.AddUnique(Tag.ToString());
-	}
-
-	return CleanedDescriptions;
-}
-
 /** Utility to get EMode from a string */
 void LexFromString(UE::Virtualization::EMode& OutValue, const FStringView& InString)
 {
@@ -63,6 +34,10 @@ void LexFromString(UE::Virtualization::EMode& OutValue, const FStringView& InStr
 	{
 		OutValue = UE::Virtualization::EMode::PackageList;
 	}
+	else if (InString == TEXT("Virtualize"))
+	{
+		OutValue = UE::Virtualization::EMode::Virtualize;
+	}
 	else if (InString == TEXT("Rehydrate"))
 	{
 		OutValue = UE::Virtualization::EMode::Rehydrate;
@@ -167,36 +142,7 @@ EInitResult FUnrealVirtualizationToolApp::Initialize()
 		return CmdLineResult;
 	}
 
-	TArray<FString> Packages;
-	if (!CurrentCommand.IsValid())
-	{
-		// Legacy code path for virtualization (which has not yet been ported to a command)	
-		switch (Mode)
-		{
-			case EMode::Changelist:
-				if (!TryParseChangelist(Packages))
-				{
-					return EInitResult::Error;
-				}
-				break;
-
-			case EMode::PackageList:
-				if (!TryParsePackageList(Packages))
-				{
-					return EInitResult::Error;
-				}
-				break;
-
-			default:
-				UE_LOG(LogVirtualizationTool, Display, TEXT("Unknown mode, cannot find packages!"));
-				return EInitResult::Error;
-				break;
-		}
-	}
-	else
-	{
-		Packages = CurrentCommand->GetPackages();
-	}
+	TArray<FString> Packages = CurrentCommand->GetPackages();
 
 	UE_LOG(LogVirtualizationTool, Display, TEXT("\tFound '%d' package file(s)"), Packages.Num());
 
@@ -214,165 +160,38 @@ bool FUnrealVirtualizationToolApp::Run()
 {
 	TRACE_CPUPROFILER_EVENT_SCOPE(Run);
 
-	if (!CurrentCommand.IsValid())
-	{
-		// Legacy code path for virtualization (which has not yet been ported to a command)	
-		TArray<FString> FinalDescriptionTags;
-
-		if (EnumHasAllFlags(ProcessOptions, EProcessOptions::Virtualize))
-		{
-			UE_LOG(LogVirtualizationTool, Display, TEXT("Running the virtualization process..."));
-
-			TArray<FText> DescriptionTags;
-
-			for (const FProject& Project : Projects)
-			{
-				TStringBuilder<128> ProjectName;
-				ProjectName << Project.GetProjectName();
-
-				UE_LOG(LogVirtualizationTool, Display, TEXT("\tChecking package(s) for the project '%s'..."), ProjectName.ToString());
-
-				FConfigFile EngineConfigWithProject;
-				if (!Project.TryLoadConfig(EngineConfigWithProject))
-				{
-					return false;
-				}
-
-				Project.RegisterMountPoints();
-
-				UE::Virtualization::FInitParams InitParams(ProjectName, EngineConfigWithProject);
-				UE::Virtualization::Initialize(InitParams, UE::Virtualization::EInitializationFlags::ForceInitialize);
-
-				TArray<FString> Packages = Project.GetAllPackages();
-
-				TArray<FText> Errors;
-				UE::Virtualization::IVirtualizationSystem::Get().TryVirtualizePackages(Packages, DescriptionTags, Errors);
-
-				if (!Errors.IsEmpty())
-				{
-					UE_LOG(LogVirtualizationTool, Error, TEXT("The virtualization process failed with the following errors:"));
-					for (const FText& Error : Errors)
-					{
-						UE_LOG(LogVirtualizationTool, Error, TEXT("\t%s"), *Error.ToString());
-					}
-					return false;
-				}
-
-				UE_LOG(LogVirtualizationTool, Display, TEXT("\tCheck complete"));
-
-				UE::Virtualization::Shutdown();
-				Project.UnRegisterMountPoints();
-			}
-
-			FinalDescriptionTags = BuildFinalTagDescriptions(DescriptionTags);
-		}
-		else
-		{
-			UE_LOG(LogVirtualizationTool, Display, TEXT("Skipping the virtualization process"));
-		}
-
-		if (Mode == EMode::Changelist)
-		{
-			if (!TrySubmitChangelist(FinalDescriptionTags))
-			{
-				return false;
-			}
-		}
+	UE_LOG(LogVirtualizationTool, Display, TEXT("Running the '%s' command..."), *CurrentCommand->GetName());
 
+	if (CurrentCommand->Run(Projects))
+	{
+		UE_LOG(LogVirtualizationTool, Display, TEXT("Command '%s' succeeded!"), *CurrentCommand->GetName());
 		return true;
 	}
 	else
 	{
-		UE_LOG(LogVirtualizationTool, Display, TEXT("Running the '%s' command..."), *CurrentCommand->GetName());
-
-		if (CurrentCommand->Run(Projects))
-		{
-			UE_LOG(LogVirtualizationTool, Display, TEXT("Command '%s' suceeeded!"), *CurrentCommand->GetName());
-			return true;
-		}
-		else
-		{
-			UE_LOG(LogVirtualizationTool, Error, TEXT("Command '%s' failed!"), *CurrentCommand->GetName());
-			return false;
-		}
+		UE_LOG(LogVirtualizationTool, Error, TEXT("Command '%s' failed!"), *CurrentCommand->GetName());
+		return false;
 	}
 }
 
 void FUnrealVirtualizationToolApp::PrintCmdLineHelp() const
 {
 	UE_LOG(LogVirtualizationTool, Display, TEXT("Usage:"));
-	UE_LOG(LogVirtualizationTool, Display, TEXT("UnrealVirtualizationTool -ClientSpecName=<name> -Mode=Changelist -Changelist=<number> [-nosubmit] [global options]"));
-	UE_LOG(LogVirtualizationTool, Display, TEXT("\t[optional]-nosubmit (the changelist will be virtualized but not submitted)"));
-	UE_LOG(LogVirtualizationTool, Display, TEXT("UnrealVirtualizationTool -ClientSpecName=<name> -Mode=PackageList -Path=<string> [global options]"));
 
 	// TODO: If the commands were registered in some way we could automate this
+	FVirtualizeCommand::PrintCmdLineHelp(); 
 	FRehydrateCommand::PrintCmdLineHelp();
 
+	UE_LOG(LogVirtualizationTool, Display, TEXT(""));
+	UE_LOG(LogVirtualizationTool, Display, TEXT("Legacy Commands:"));
+	UE_LOG(LogVirtualizationTool, Display, TEXT("-Mode=Changelist -ClientSpecName=<name> [optional] -Changelist=<number> -nosubmit [optional]"));
+	UE_LOG(LogVirtualizationTool, Display, TEXT("-Mode=PackageList -Path=<string>"));
+	
+	UE_LOG(LogVirtualizationTool, Display, TEXT(""));
 	UE_LOG(LogVirtualizationTool, Display, TEXT("Global Options:"));
 	UE_LOG(LogVirtualizationTool, Display, TEXT("\t-verbose (all log messages with display verbosity will be displayed, not just LogVirtualizationTool)"));
 }
 
-bool FUnrealVirtualizationToolApp::TrySubmitChangelist(const TArray<FString>& DescriptionTags)
-{
-	if (!EnumHasAllFlags(ProcessOptions, EProcessOptions::Submit))
-	{
-		UE_LOG(LogVirtualizationTool, Display, TEXT("Skipping submit of changelist '%s' due to cmdline options"), *ChangelistNumber);
-		return true;
-	}
-
-	UE_LOG(LogVirtualizationTool, Display, TEXT("Attempting to submit the changelist '%s'"), *ChangelistNumber);
-
-	if (!TryConnectToSourceControl())
-	{
-		UE_LOG(LogVirtualizationTool, Error, TEXT("Submit failed, cannot find a valid source control provider"));
-		return false;
-	}
-
-	if (!ChangelistToSubmit.IsValid())
-	{
-		// This should not be possible, the check and error message is to guard against potential future problems only.
-		UE_LOG(LogVirtualizationTool, Error, TEXT("Submit failed, could not find the changelist"));
-		return false;
-	}
-
-	FSourceControlChangelistRef Changelist = ChangelistToSubmit.ToSharedRef();
-	FSourceControlChangelistStatePtr ChangelistState = SCCProvider->GetState(Changelist, EStateCacheUsage::Use);
-
-	if (!ChangelistState.IsValid())
-	{
-		UE_LOG(LogVirtualizationTool, Error, TEXT("Submit failed, failed to find the state for the changelist"));
-		return false;
-	}
-
-	TSharedRef<FCheckIn, ESPMode::ThreadSafe> CheckInOperation = ISourceControlOperation::Create<FCheckIn>();
-
-	// Grab the original changelist description then append our tags afterwards
-	TStringBuilder<512> Description;
-	Description << ChangelistState->GetDescriptionText().ToString();
-
-	for (const FString& Tag : DescriptionTags)
-	{
-		Description << TEXT("\n") << Tag;
-	}
-
-	CheckInOperation->SetDescription(FText::FromString(Description.ToString()));
-
-	if (SCCProvider->Execute(CheckInOperation, Changelist) == ECommandResult::Succeeded)
-	{
-		UE_LOG(LogVirtualizationTool, Display, TEXT("%s"), *CheckInOperation->GetSuccessMessage().ToString());	
-		return true;
-	}
-	else
-	{
-		// Even when log suppression is active we still show errors to the users and as the source control
-		// operation should have logged the problem as an error the user will see it. This means we don't 
-		// have to extract it from CheckInOperation .
-		UE_LOG(LogVirtualizationTool, Error, TEXT("Submit failed, please check the log!"));
-
-		return false;
-	}
-}
-
 bool FUnrealVirtualizationToolApp::TryLoadModules()
 {
 	if (FModuleManager::Get().LoadModule(TEXT("Virtualization"), ELoadModuleFlags::LogFailures) == nullptr)
@@ -411,33 +230,6 @@ bool FUnrealVirtualizationToolApp::TryInitEnginePlugins()
 	return true;
 }
 
-bool FUnrealVirtualizationToolApp::TryConnectToSourceControl()
-{
-	TRACE_CPUPROFILER_EVENT_SCOPE(TryConnectToSourceControl);
-
-	if (SCCProvider.IsValid())
-	{
-		// Already connected so just return
-		return true;
-	}
-
-	UE_LOG(LogVirtualizationTool, Log, TEXT("Trying to connect to source control..."));
-
-	FSourceControlInitSettings SCCSettings(FSourceControlInitSettings::EBehavior::OverrideAll);
-	SCCSettings.AddSetting(TEXT("P4Client"), ClientSpecName);
-
-	SCCProvider = ISourceControlModule::Get().CreateProvider(FName("Perforce"), TEXT("UnrealVirtualizationTool"), SCCSettings);
-	if (SCCProvider.IsValid())
-	{
-		return true;
-	}
-	else
-	{
-		UE_LOG(LogVirtualizationTool, Error, TEXT("Failed to create a perforce connection"));
-		return false;
-	}
-}
-
 EInitResult FUnrealVirtualizationToolApp::TryParseCmdLine()
 {
 	TRACE_CPUPROFILER_EVENT_SCOPE(TryParseCmdLine);
@@ -460,15 +252,6 @@ EInitResult FUnrealVirtualizationToolApp::TryParseCmdLine()
 		return EInitResult::EarlyOut;
 	}
 
-	// First parse the command line options that can apply to all modes
-
-	if (!FParse::Value(CmdLine, TEXT("-ClientSpecName="), ClientSpecName))
-	{
-		UE_LOG(LogVirtualizationTool, Error, TEXT("Failed to find cmdline switch 'ClientSpecName', this is a required parameter!"));
-		PrintCmdLineHelp();
-		return EInitResult::Error;
-	}
-
 	if (FParse::Param(CmdLine, TEXT("Verbose")))
 	{
 		UE_LOG(LogVirtualizationTool, Display, TEXT("Cmdline parameter '-Verbose' found, no longer supressing Display log messages!"));
@@ -494,10 +277,16 @@ EInitResult FUnrealVirtualizationToolApp::TryParseCmdLine()
 	switch (Mode)
 	{
 		case EMode::Changelist:
-			return TryParseChangelistCmdLine(CmdLine);
+			CurrentCommand = CreateCommand<FVirtualizeLegacyChangeListCommand>(ModeAsString, CmdLine);
+			return CurrentCommand.IsValid() ? EInitResult::Success : EInitResult::Error;
 			break;
 		case EMode::PackageList:
-			return TryParsePackageListCmdLine(CmdLine);
+			CurrentCommand = CreateCommand<FVirtualizeLegacyPackageListCommand>(ModeAsString, CmdLine);
+			return CurrentCommand.IsValid() ? EInitResult::Success : EInitResult::Error;
+			break;
+		case EMode::Virtualize:
+			CurrentCommand = CreateCommand<FVirtualizeCommand>(ModeAsString, CmdLine);
+			return CurrentCommand.IsValid() ? EInitResult::Success : EInitResult::Error;
 			break;
 		case EMode::Rehydrate:
 			CurrentCommand = CreateCommand<FRehydrateCommand>(ModeAsString, CmdLine);
@@ -513,163 +302,6 @@ EInitResult FUnrealVirtualizationToolApp::TryParseCmdLine()
 	}
 }
 
-EInitResult	FUnrealVirtualizationToolApp::TryParseChangelistCmdLine(const TCHAR* CmdLine)
-{
-	if (FParse::Value(CmdLine, TEXT("-Changelist="), ChangelistNumber))
-	{
-		// Optional switches
-		if (FParse::Param(CmdLine, TEXT("NoSubmit")))
-		{
-			UE_LOG(LogVirtualizationTool, Display, TEXT("Cmdline parameter '-NoSubmit' found, the changelist will be virtualized but not submitted!"));
-		}
-		else
-		{
-			ProcessOptions |= EProcessOptions::Submit;
-		}
-
-		UE_LOG(LogVirtualizationTool, Display, TEXT("Attempting to virtualize changelist '%s'"), *ChangelistNumber);
-		
-		return EInitResult::Success;
-	}
-	else
-	{
-		UE_LOG(LogVirtualizationTool, Error, TEXT("Failed to find cmdline switch 'Changelist', this is a required parameter for the 'Changelist' mode!"));
-		PrintCmdLineHelp();
-		return EInitResult::Error;
-	}
-}
-
-EInitResult FUnrealVirtualizationToolApp::TryParsePackageListCmdLine(const TCHAR* CmdLine)
-{
-	if (FParse::Value(CmdLine, TEXT("-Path="), PackageListPath))
-	{
-		UE_LOG(LogVirtualizationTool, Display, TEXT("Virtualizing packages found in package list: '%s'"), *PackageListPath);
-		return EInitResult::Success;
-	}
-	else
-	{
-		UE_LOG(LogVirtualizationTool, Error, TEXT("Failed to find cmdline switch 'Path', this is a required parameter for the 'PackageList mode!"));
-		PrintCmdLineHelp();
-		return EInitResult::Error;
-	}	
-}
-
-bool FUnrealVirtualizationToolApp::TryParseChangelist(TArray<FString>& OutPackages)
-{
-	TRACE_CPUPROFILER_EVENT_SCOPE(TryParseChangelist);
-
-	if (!TryConnectToSourceControl())
-	{
-		return false;
-	}
-
-	UE_LOG(LogVirtualizationTool, Display, TEXT("Attempting to parse changelist '%s' in workspace '%s'"), *ChangelistNumber, *ClientSpecName);
-
-	if (!SCCProvider.IsValid())
-	{
-		UE_LOG(LogVirtualizationTool, Error, TEXT("No valid source control connection found!"));
-		return false;
-	}
-
-	SCCProvider->Init(true);
-
-	if (!SCCProvider->UsesChangelists())
-	{
-		UE_LOG(LogVirtualizationTool, Error, TEXT("The source control provider does not support the use of changelists"));
-		return false;
-	}
-
-	TArray<FSourceControlChangelistRef> Changelists = SCCProvider->GetChangelists(EStateCacheUsage::ForceUpdate);
-	if (Changelists.IsEmpty())
-	{
-		UE_LOG(LogVirtualizationTool, Error, TEXT("Failed to find any changelists"));
-		return false;
-	}
-
-	TArray<FSourceControlChangelistStateRef> ChangelistsStates;
-	if (SCCProvider->GetState(Changelists, ChangelistsStates, EStateCacheUsage::Use) != ECommandResult::Succeeded)
-	{
-		UE_LOG(LogVirtualizationTool, Error, TEXT("Failed to find changelist data"));
-		return false;
-	}
-
-	for (FSourceControlChangelistStateRef& ChangelistState : ChangelistsStates)
-	{
-		const FText DisplayText = ChangelistState->GetDisplayText();
-
-		if (ChangelistNumber == DisplayText.ToString())
-		{
-			TSharedRef<FUpdatePendingChangelistsStatus, ESPMode::ThreadSafe> Operation = ISourceControlOperation::Create<FUpdatePendingChangelistsStatus>();
-
-			// TODO: Updating only the CL we want does not currently work and even if it did we still end up with a pointless
-			// p4 changes command before updating the files. Given we know the changelist number via FSourceControlChangelistRef
-			// we should be able to just request the file states be updated.
-			// This is also a lot of code to write for a simple "give me all files in a changelist" operation, if we don't add
-			// support directly in the API we should move this to a utility namespace in the source control module.
-
-			FSourceControlChangelistRef Changelist = ChangelistState->GetChangelist();
-			Operation->SetChangelistsToUpdate(MakeArrayView(&Changelist, 1));
-			Operation->SetUpdateFilesStates(true);
-			
-			if (SCCProvider->Execute(Operation, EConcurrency::Synchronous) != ECommandResult::Succeeded)
-			{
-				UE_LOG(LogVirtualizationTool, Error, TEXT("Failed to find the files in changelist '%s'"), *ChangelistNumber);
-				return false;
-			}
-
-			const TArray<FSourceControlStateRef>& FilesinChangelist = ChangelistState->GetFilesStates();
-			for (const FSourceControlStateRef& FileState : FilesinChangelist)
-			{
-				if (::IsPackageFile(FileState->GetFilename()))
-				{
-					OutPackages.Add(FileState->GetFilename());
-				}
-				else
-				{
-					UE_LOG(LogVirtualizationTool, Log, TEXT("Ignoring non-package file '%s'"), *FileState->GetFilename());
-				}
-			}
-
-			ChangelistToSubmit = Changelist;
-
-			return true;
-		}
-	}
-
-	UE_LOG(LogVirtualizationTool, Error, TEXT("Failed to find the changelist '%s'"), *ChangelistNumber);
-	return false;
-}
-
-bool FUnrealVirtualizationToolApp::TryParsePackageList(TArray<FString>& OutPackages)
-{
-	TRACE_CPUPROFILER_EVENT_SCOPE(TrySortFilesByProject);
-
-	UE_LOG(LogVirtualizationTool, Display, TEXT("Parsing the package list..."));
-
-	if (!IFileManager::Get().FileExists(*PackageListPath))
-	{
-		UE_LOG(LogVirtualizationTool, Error, TEXT("\tThe package list '%s' does not exist"), *PackageListPath);
-		return false;
-	}
-
-	if (FFileHelper::LoadFileToStringArray(OutPackages, *PackageListPath))
-	{
-		// We don't have control over how the package list was generated so make sure that the paths
-		// are in the format that we want.
-		for (FString& PackagePath : OutPackages)
-		{
-			FPaths::NormalizeFilename(PackagePath);
-		}
-
-		return true;
-	}
-	else
-	{
-		UE_LOG(LogVirtualizationTool, Error, TEXT("\tFailed to parse the package list '%s'"), *PackageListPath);
-		return false;
-	}
-}
-
 bool FUnrealVirtualizationToolApp::TrySortFilesByProject(const TArray<FString>& Packages)
 {
 	TRACE_CPUPROFILER_EVENT_SCOPE(TrySortFilesByProject);
diff --git a/Engine/Source/Programs/UnrealVirtualizationTool/Private/UnrealVirtualizationToolApp.h b/Engine/Source/Programs/UnrealVirtualizationTool/Private/UnrealVirtualizationToolApp.h
--- a/Engine/Source/Programs/UnrealVirtualizationTool/Private/UnrealVirtualizationToolApp.h
+++ b/Engine/Source/Programs/UnrealVirtualizationTool/Private/UnrealVirtualizationToolApp.h
@@ -4,12 +4,9 @@
 
 #include "Containers/Array.h"
 #include "Containers/UnrealString.h"
-#include "ISourceControlChangelist.h" // TODO
 #include "ProjectFiles.h"
 #include "Templates/UniquePtr.h"
 
-class ISourceControlProvider;
-
 namespace UE::Virtualization
 {
 
@@ -21,33 +18,21 @@ enum class EMode :uint32
 	/** Error condition */
 	Unknown = 0,
 
-	/// Virtualization
+	// Legacy (should be phased out)
 
 	/** Virtualize and submit a given changelist */
 	Changelist,
 	/** Virtualize a list of packages provided by text file */
 	PackageList,
 
-	/// Rehydration
-
-	/** Rehydrate one or more packages */
-	Rehydrate,
-};
+	/// New Style Commands
 
-// Note that the enum doesn't give us huge value right now, it has been provided
-// in case that we expand with additional functionality in the future.
+	/** Virtualize one or more packages */
+	Virtualize,
 
-/** A bitfield of the various operations that the virtualization process supports */
-enum class EProcessOptions : uint32
-{
-	/** No options */
-	None = 0,
-	/** Virtualize the packages in the provided changelist */
-	Virtualize = 1 << 0,
-	/** Submit the changelist */
-	Submit = 1 << 1
+	/** Rehydrate one or more packages */
+	Rehydrate
 };
-ENUM_CLASS_FLAGS(EProcessOptions);
 
 /** The result of FUnrealVirtualizationToolApp initialization */
 enum class EInitResult
@@ -73,18 +58,10 @@ private:
 
 	void PrintCmdLineHelp() const;
 
-	bool TrySubmitChangelist(const TArray<FString>& DescriptionTags);
-
 	bool TryLoadModules();
 	bool TryInitEnginePlugins();
-	bool TryConnectToSourceControl();
-
+	
 	EInitResult TryParseCmdLine();
-	EInitResult TryParseChangelistCmdLine(const TCHAR* CmdLine);
-	EInitResult	TryParsePackageListCmdLine(const TCHAR* CmdLine);
-
-	bool TryParseChangelist(TArray<FString>& OutPackages);
-	bool TryParsePackageList(TArray<FString>& OutPackages);
 
 	bool TrySortFilesByProject(const TArray<FString>& Packages);
 
@@ -97,30 +74,13 @@ private:
 	/** Override used to suppress log messages */
 	TUniquePtr<FFeedbackContext> OutputDeviceOverride;
 
-	/** The source control provider that the tool uses */
-	TUniquePtr<ISourceControlProvider> SCCProvider;
-
-	/** Pointer to the changelist that should be submitted */
-	FSourceControlChangelistPtr ChangelistToSubmit;
-
 	/** Data structure holding the files in the changelist sorted by project and then by plugin*/
 	TArray<FProject> Projects;
 
-	/** Name of the client spec (workspace) passed in on the command line*/
-	FString ClientSpecName;
-
 	/** The mode for the application to run */
 	EMode Mode = EMode::Unknown;
 
-	/** Bitfield control the various options that should be run */
-	EProcessOptions ProcessOptions = EProcessOptions::Virtualize;
-
-	/** The number of the changelist being submitted, used with EMode::Changelist  */
-	FString ChangelistNumber;
-
-	/** The path to the list of packages to virtualized, used with EMode::PackageList */
-	FString PackageListPath;
-
+	/** The command being processed */
 	TUniquePtr<FCommand> CurrentCommand;
 };
 
