Fix a shutdown crash in UVT if the virtualization command returns an error

#rnx

- Moved the shutdown code to scope exit so that they will continue to trigger on any early out paths.

diff --git a/Engine/Source/Programs/UnrealVirtualizationTool/Private/Commands/VirtualizeCommand.cpp b/Engine/Source/Programs/UnrealVirtualizationTool/Private/Commands/VirtualizeCommand.cpp
--- a/Engine/Source/Programs/UnrealVirtualizationTool/Private/Commands/VirtualizeCommand.cpp
+++ b/Engine/Source/Programs/UnrealVirtualizationTool/Private/Commands/VirtualizeCommand.cpp
@@ -141,9 +141,19 @@ bool FVirtualizeCommand::Run(const TArray<FProject>& Projects)
 
 		Project.RegisterMountPoints();
 
+		ON_SCOPE_EXIT
+		{
+			Project.UnRegisterMountPoints();
+		};
+
 		UE::Virtualization::FInitParams InitParams(ProjectName, EngineConfigWithProject);
 		UE::Virtualization::Initialize(InitParams, UE::Virtualization::EInitializationFlags::ForceInitialize);
 
+		ON_SCOPE_EXIT
+		{
+			UE::Virtualization::Shutdown();
+		};
+
 		TArray<FString> ProjectPackages = Project.GetAllPackages();
 
 		TArray<FText> Errors;
@@ -160,9 +170,6 @@ bool FVirtualizeCommand::Run(const TArray<FProject>& Projects)
 		}
 
 		UE_LOG(LogVirtualizationTool, Display, TEXT("\tCheck complete"));
-
-		UE::Virtualization::Shutdown();
-		Project.UnRegisterMountPoints();
 	}
 
 	if (bShouldSubmitChangelist)
