Improved a code comment to be less confusing based on licensee feedback.

#rnx

diff --git a/Engine/Source/Programs/UnrealVirtualizationTool/Private/Commands/CommandBase.cpp b/Engine/Source/Programs/UnrealVirtualizationTool/Private/Commands/CommandBase.cpp
--- a/Engine/Source/Programs/UnrealVirtualizationTool/Private/Commands/CommandBase.cpp
+++ b/Engine/Source/Programs/UnrealVirtualizationTool/Private/Commands/CommandBase.cpp
@@ -220,6 +220,10 @@ bool FCommand::TryParseChangelist(FStringView ClientSpecName, FStringView Change
 		return false;
 	}
 
+	// TODO: At the moment we have to poll for all changelists under the current workspace and then iterate over
+	// them to find the one we want.
+	// We need to add better support in the API to go from a changelist number to a FSourceControlChangelistRef
+	// directly.
 	TArray<FSourceControlChangelistStateRef> ChangelistsStates;
 	if (SCCProvider->GetState(Changelists, ChangelistsStates, EStateCacheUsage::Use) != ECommandResult::Succeeded)
 	{
@@ -235,12 +239,6 @@ bool FCommand::TryParseChangelist(FStringView ClientSpecName, FStringView Change
 		{
 			TSharedRef<FUpdatePendingChangelistsStatus, ESPMode::ThreadSafe> Operation = ISourceControlOperation::Create<FUpdatePendingChangelistsStatus>();
 
-			// TODO: Updating only the CL we want does not currently work and even if it did we still end up with a pointless
-			// p4 changes command before updating the files. Given we know the changelist number via FSourceControlChangelistRef
-			// we should be able to just request the file states be updated.
-			// This is also a lot of code to write for a simple "give me all files in a changelist" operation, if we don't add
-			// support directly in the API we should move this to a utility namespace in the source control module.
-
 			FSourceControlChangelistRef Changelist = ChangelistState->GetChangelist();
 			Operation->SetChangelistsToUpdate(MakeArrayView(&Changelist, 1));
 			Operation->SetUpdateFilesStates(true);
