Change mention of "UnrealHeaderTool" to "This app" in comment.

#rnx

diff --git a/Engine/Source/Programs/BlankProgram/BlankProgram.Target.cs b/Engine/Source/Programs/BlankProgram/BlankProgram.Target.cs
--- a/Engine/Source/Programs/BlankProgram/BlankProgram.Target.cs
+++ b/Engine/Source/Programs/BlankProgram/BlankProgram.Target.cs
@@ -31,7 +31,7 @@ public class BlankProgramTarget : TargetRules
 		bCompileAgainstApplicationCore = false;
 		bCompileICU = false;
 
-		// UnrealHeaderTool is a console application, not a Windows app (sets entry point to main(), instead of WinMain())
+		// This app is a console application, not a Windows app (sets entry point to main(), instead of WinMain())
 		bIsBuildingConsoleApplication = true;
 	}
 }
diff --git a/Engine/Source/Programs/HeadlessChaos/HeadlessChaos.Target.cs b/Engine/Source/Programs/HeadlessChaos/HeadlessChaos.Target.cs
--- a/Engine/Source/Programs/HeadlessChaos/HeadlessChaos.Target.cs
+++ b/Engine/Source/Programs/HeadlessChaos/HeadlessChaos.Target.cs
@@ -17,8 +17,8 @@ public class HeadlessChaosTarget : TestTargetRules
 
         bUseLoggingInShipping = true;
 
-        // UnrealHeaderTool is a console application, not a Windows app (sets entry point to main(), instead of WinMain())
-        bIsBuildingConsoleApplication = true;
+		// This app is a console application, not a Windows app (sets entry point to main(), instead of WinMain())
+		bIsBuildingConsoleApplication = true;
 
 		GlobalDefinitions.Add("CHAOS_SERIALIZE_OUT=1");
 
diff --git a/Engine/Source/Programs/LiveCodingConsole/LiveCodingConsole.Target.cs b/Engine/Source/Programs/LiveCodingConsole/LiveCodingConsole.Target.cs
--- a/Engine/Source/Programs/LiveCodingConsole/LiveCodingConsole.Target.cs
+++ b/Engine/Source/Programs/LiveCodingConsole/LiveCodingConsole.Target.cs
@@ -29,8 +29,5 @@ public class LiveCodingConsoleTarget : TargetRules
 
 		// ICU is needed for regex during click to source code
 		bCompileICU = true;
-
-		// UnrealHeaderTool is a console application, not a Windows app (sets entry point to main(), instead of WinMain())
-		bIsBuildingConsoleApplication = false;
 	}
 }
diff --git a/Engine/Source/Programs/UnrealHeaderTool/UnrealHeaderTool.Target.cs b/Engine/Source/Programs/UnrealHeaderTool/UnrealHeaderTool.Target.cs
--- a/Engine/Source/Programs/UnrealHeaderTool/UnrealHeaderTool.Target.cs
+++ b/Engine/Source/Programs/UnrealHeaderTool/UnrealHeaderTool.Target.cs
@@ -31,7 +31,7 @@ public class UnrealHeaderToolTarget : TargetRules
 		bCompileWithPluginSupport = true;
 		bBuildDeveloperTools = true;
 
-		// UnrealHeaderTool is a console application, not a Windows app (sets entry point to main(), instead of WinMain())
+		// This app is a console application, not a Windows app (sets entry point to main(), instead of WinMain())
 		bIsBuildingConsoleApplication = true;
 
         GlobalDefinitions.Add("HACK_HEADER_GENERATOR=1");
diff --git a/Engine/Source/Programs/UnrealLightmass/UnrealLightmass.Target.cs b/Engine/Source/Programs/UnrealLightmass/UnrealLightmass.Target.cs
--- a/Engine/Source/Programs/UnrealLightmass/UnrealLightmass.Target.cs
+++ b/Engine/Source/Programs/UnrealLightmass/UnrealLightmass.Target.cs
@@ -36,7 +36,7 @@ public class UnrealLightmassTarget : TargetRules
 			bBuildDeveloperTools = true;
 		}
 
-		// UnrealHeaderTool is a console application, not a Windows app (sets entry point to main(), instead of WinMain())
+		// This app is a console application, not a Windows app (sets entry point to main(), instead of WinMain())
 		bIsBuildingConsoleApplication = true;
 
 		// Disable logging, lightmass will create its own unique logging file
diff --git a/Engine/Source/Programs/UnrealVirtualizationTool/UnrealVirtualizationTool.Target.cs b/Engine/Source/Programs/UnrealVirtualizationTool/UnrealVirtualizationTool.Target.cs
--- a/Engine/Source/Programs/UnrealVirtualizationTool/UnrealVirtualizationTool.Target.cs
+++ b/Engine/Source/Programs/UnrealVirtualizationTool/UnrealVirtualizationTool.Target.cs
@@ -25,7 +25,7 @@ public class UnrealVirtualizationToolTarget : TargetRules
 		bCompileAgainstEngine = false;
 		bUsesSlate = false;
 
-		// UnrealHeaderTool is a console application, not a Windows app (sets entry point to main(), instead of WinMain())
+		// This app is a console application, not a Windows app (sets entry point to main(), instead of WinMain())
 		bIsBuildingConsoleApplication = true;
 
 		// Enable Developer plugins
