Fix the UnrealVirtualizationTool to work with the new interactive login system for the cloud DDC

#rnx

- The DDC module can now invoke the 'DesktopPlatform' module when trying to set up the interactive ddc login. At the moment the 'DerivedDataCache' module only links to the includes with no dependency on the module itself so as a short term fix I am adding a dependency from the tool itself.
- The interactive login process relies on our basic OS message dialogs. On windows this requires that the .rc file in the 'Launch' module be included otherwise the message dialog will not be shown. We don't want a dependency on that module directly as it is super heavy weight but we can add our own .rc file for the tool and inside that include the .rc file of the 'Launch' module.

diff --git a/Engine/Source/Programs/UnrealVirtualizationTool/Resources/Windows/UnrealVirtualizationTool.rc b/Engine/Source/Programs/UnrealVirtualizationTool/Resources/Windows/UnrealVirtualizationTool.rc
new file mode 100644
--- /dev/null
+++ b/Engine/Source/Programs/UnrealVirtualizationTool/Resources/Windows/UnrealVirtualizationTool.rc
@@ -0,0 +1,5 @@
+// Copyright Epic Games, Inc. All Rights Reserved.
+
+#include <windows.h>
+#include "Runtime/Launch/Resources/Version.h"
+#include "Runtime/Launch/Resources/Windows/PCLaunch.rc"
\ No newline at end of file
diff --git a/Engine/Source/Programs/UnrealVirtualizationTool/UnrealVirtualizationTool.Build.cs b/Engine/Source/Programs/UnrealVirtualizationTool/UnrealVirtualizationTool.Build.cs
--- a/Engine/Source/Programs/UnrealVirtualizationTool/UnrealVirtualizationTool.Build.cs
+++ b/Engine/Source/Programs/UnrealVirtualizationTool/UnrealVirtualizationTool.Build.cs
@@ -15,7 +15,8 @@ public class UnrealVirtualizationTool : ModuleRules
 				"Projects",
 				"ApplicationCore",
 				"SourceControl",
-				"PerforceSourceControl"
+				"PerforceSourceControl",
+				"DesktopPlatform"
 			});
 
 		PublicIncludePaths.Add("Runtime/Launch/Public");
