Clean up FCommand::TryConnectToSourceControl by adding an overload with no parameters.

#rnx

- This makes calling it without a valid client spec alittle cleaner as we no longer need to pass in a dummy FStringView/String

diff --git a/Engine/Source/Programs/UnrealVirtualizationTool/Private/Commands/CommandBase.cpp b/Engine/Source/Programs/UnrealVirtualizationTool/Private/Commands/CommandBase.cpp
--- a/Engine/Source/Programs/UnrealVirtualizationTool/Private/Commands/CommandBase.cpp
+++ b/Engine/Source/Programs/UnrealVirtualizationTool/Private/Commands/CommandBase.cpp
@@ -283,7 +283,7 @@ FString FCommand::FindClientSpecForChangelist(FStringView ChangelistNumber)
 	UE_LOG(LogVirtualizationTool, Display, TEXT("\tAttempting to find the workspace for '%.*s'"),
 		ChangelistNumber.Len(), ChangelistNumber.GetData());
 
-	if (!TryConnectToSourceControl(FStringView()))
+	if (!TryConnectToSourceControl())
 	{
 		return FString();
 	}
diff --git a/Engine/Source/Programs/UnrealVirtualizationTool/Private/Commands/CommandBase.h b/Engine/Source/Programs/UnrealVirtualizationTool/Private/Commands/CommandBase.h
--- a/Engine/Source/Programs/UnrealVirtualizationTool/Private/Commands/CommandBase.h
+++ b/Engine/Source/Programs/UnrealVirtualizationTool/Private/Commands/CommandBase.h
@@ -62,6 +62,10 @@ protected: // Common SourceControl Code
 	 * @return True if the provider was created, otherwise false.
 	 */
 	bool TryConnectToSourceControl(FStringView ClientSpecName);
+	bool TryConnectToSourceControl()
+	{
+		return TryConnectToSourceControl(FStringView());
+	}
 
 	bool TryCheckOutFilesForProject(FStringView ClientSpecName, FStringView ProjectRoot, const TArray<FString>& ProjectPackages);
 	
