[UVT] Improve error message to reflect the actual problem (we could not create a perforce source control provider object rather than the actual connection failing)

#rnx

diff --git a/Engine/Source/Programs/UnrealVirtualizationTool/Private/Commands/CommandBase.cpp b/Engine/Source/Programs/UnrealVirtualizationTool/Private/Commands/CommandBase.cpp
--- a/Engine/Source/Programs/UnrealVirtualizationTool/Private/Commands/CommandBase.cpp
+++ b/Engine/Source/Programs/UnrealVirtualizationTool/Private/Commands/CommandBase.cpp
@@ -129,7 +129,7 @@ bool FCommand::TryConnectToSourceControl(FStringView ClientSpecName)
 	}
 	else
 	{
-		UE_LOG(LogVirtualizationTool, Error, TEXT("Failed to create a perforce connection"));
+		UE_LOG(LogVirtualizationTool, Error, TEXT("Failed to instantiate a perforce revision control connection"));
 		return false;
 	}
 }
