[UVT] Improve error handling when failing to establish a perforce connection.

#rnx

- As ISourceControlProvide::Init does not return if a connection was established or not we need to also call ISourceControlProvide::IsAvailable afterwards to check if we actually have a connection.
- This lets us fail early with a more specific error message rather than waiting for the first command to fail especially as our commands rarely print useful errors if the connection failed.

diff --git a/Engine/Source/Programs/UnrealVirtualizationTool/Private/Commands/CommandBase.cpp b/Engine/Source/Programs/UnrealVirtualizationTool/Private/Commands/CommandBase.cpp
--- a/Engine/Source/Programs/UnrealVirtualizationTool/Private/Commands/CommandBase.cpp
+++ b/Engine/Source/Programs/UnrealVirtualizationTool/Private/Commands/CommandBase.cpp
@@ -125,7 +125,17 @@ bool FCommand::TryConnectToSourceControl(FStringView ClientSpecName)
 	if (SCCProvider.IsValid())
 	{
 		SCCProvider->Init(true);
-		return true;
+
+		if (SCCProvider->IsAvailable())
+		{
+			return true;
+		}
+		else
+		{
+			UE_LOG(LogVirtualizationTool, Error, TEXT("Failed to establish a perforce connection"));
+			return false;
+		}
+		
 	}
 	else
 	{
