[UVT] Improve the robustness of the tool wrt ensuring that the correct project config settings are loaded.

### Problem
- A lot of systems rely on the engine to know which project is being loaded on start up so that the correct config settings can be applied however so far the UVT has not required the user to provide the desired project. This is because in theory, the tool could be given a list of packages or a changelist to process that contain packages from more than one project and if that happens we need to handle it correctly.
- The tool currently groups the provided packages by project, then runs the requested command on each project one at a time, initializing the VA system with that projects config files, then shutting it done once the project packages are processed. This allows us to pick up the correct settings for VA but other systems that we rely on, such as the DDC do not allow this sort of dynamic initialization and so may not have the correct settings.

### Solution
- Before the projects packages are processed we check to see if the project applied to the process during initialization (if any) matches the project to be processed. If not we launch a new version of the tool with the correct project settings and have the child process perform the main processing. The output of each child process (or any projects we were able to run in process) are then provided to the command for a final "post projects" phase.
- This means we can support multiple packages from multiple processes AND not require the user type out the project name when using the tool from the cmdline.
- Although launching a child process only adds a few seconds of overhead (per project) any tool or process that make use of UVT should be updated to provide the project where possible.

### Changes

### FUnrealVirtualizationToolApp
- As explained in the problem section, we now need to support launching a new instance of the tool as a child process.
- We communicate with the child process by providing an input file via the cmdline, which is then parsed and used to recreate the original command and project settings. Once the command has run in the child process it will write its output to an output file which the calling process can then read back in.
- So our new logic is to check the project that the tool was started with (in case the user provided a project path in the cmdline) against each project that we need to process. If the paths match then we can just process the project, if they do not match we launch a child process with the correct path and retrieve the output once done.
-- If a child process detects an incorrect path we just error out to avoid endless recursive process launching.
- Output from the child process is clearly logged to the user as coming from a child process to make debugging easier.

### FCommand
- Added ToJson/FromJson methods to FCommand, which simulates the same interface as though it was using the json serialization macros, however we want more control over things so do the serialization manually via FCommand::Serialize
-- Serialization is only used if the command is sent to a child process (if the current project is wrong), in which case the command should serialize out all info that it derived during setup that it will need to process correctly. Usually this consists of the options that were parsed from the command line. Then when the child process is run we will recreate the command via the same serialize method.
- Added a new class FProcessPipes, which acts as a wrapper around the pipes we can create by calling FPlatformProcess::CreatePipe which are returned as raw pointers. This ensures that the pipes are cleaned up once the FProcessPipes object goes out of scope.
- The ::Run method has been split into two, ::ProcessProject and ::ProcessOutput
-- ::ProcessProject is called once per project and allows the command to return output via the Output parameter. This can be call in process or via a child process if the project needs to be set.
-- ::ProcessOutput is called once all projects have been processed and provides an array which is a collection of the output from each call to ::ProcessProject. This is always called in process.
-- The output is passed around by the base type FCommandOutput. It is expected that commands derive their own type and as we do not mix command types, they should know how to cast the output back to the correct type when ::ProcessOutput is called. This is a bit messy but to go further risks over engineering given it is unlikely that more commands will be added to the tool. If we do expand the tool in the future then this can be revisited.

### FCommandOutput
- A struct used to represent the output of a command.
- If the command is run as a child process then this struct will be serialized out to disk (via json) then loaded back in by the calling process so that the output of each child process can be combined.
- Commands are expected to derive their own output structure and deal with the json serialization. To make this easier we provide a JSON_SERIALIZE_PARENT macro to allow for serialization of inheritance chains.
- This system does require that commands know which type of output to cast FCommandOutput to. At the moment we do not provide any type safety.

- Add json serialization to FProject and it's child struct FPlugin.
- FProject::GetProjectFilePath now returns as a reference to a FString rather than a FStringView.

diff --git a/Engine/Source/Programs/UnrealVirtualizationTool/Private/Commands/CommandBase.h b/Engine/Source/Programs/UnrealVirtualizationTool/Private/Commands/CommandBase.h
--- a/Engine/Source/Programs/UnrealVirtualizationTool/Private/Commands/CommandBase.h
+++ b/Engine/Source/Programs/UnrealVirtualizationTool/Private/Commands/CommandBase.h
@@ -5,19 +5,46 @@
 #include "Containers/Array.h"
 #include "Containers/StringView.h"
 #include "Containers/UnrealString.h"
+#include "Dom/JsonObject.h"
+#include "Serialization/JsonSerializerMacros.h"
+#include "Serialization/JsonWriter.h"
 #include "Templates/SharedPointer.h"
 #include "Templates/UniquePtr.h"
 
+struct FJsonSerializerBase;
+
 class ISourceControlProvider;
 class ISourceControlChangelist;
 
 typedef TSharedPtr<ISourceControlChangelist, ESPMode::ThreadSafe> FSourceControlChangelistPtr;
 
+
 namespace UE::Virtualization
 {
 
 class FProject;
 
+// The JSON serialization macros do not support inheritance properly as far as I can tell, this helps work around that.
+// Currently it stores all levels of the inheritance chain in the same json object, I'd want to change this to have an
+// object per level before considering moving this to JsonSerializerMacros.h
+#define JSON_SERIALIZE_PARENT(ParentType) \
+	ParentType::Serialize(Serializer, true);
+
+struct FCommandOutput : public FJsonSerializable
+{
+	FCommandOutput() = default;
+	FCommandOutput(FStringView InProjectName)
+		: ProjectName(InProjectName)
+	{
+	}
+
+	BEGIN_JSON_SERIALIZER
+		JSON_SERIALIZE("ProjectName", ProjectName);
+	END_JSON_SERIALIZER
+
+	FString ProjectName;
+};
+
 /** The base class to derive new commands from */
 class FCommand
 {
@@ -27,7 +54,47 @@ public:
 
 	virtual bool Initialize(const TCHAR* CmdLine) = 0;
 
-	virtual bool Run(const TArray<FProject>& Projects) = 0;
+	void ToJson(TSharedRef<TJsonWriter<> >& JsonWriter) const
+	{
+		FJsonSerializerWriter<> Serializer(JsonWriter);
+		const_cast<FCommand*>(this)->Serialize(Serializer);
+	}
+
+	bool FromJson(TSharedPtr<FJsonObject> JsonObject)
+	{
+		if (JsonObject.IsValid())
+		{
+			FJsonSerializerReader Serializer(JsonObject);
+			Serialize(Serializer);
+			return true;
+		}
+		return false;
+	}
+
+	virtual void Serialize(FJsonSerializerBase& Serializer) = 0;
+
+	/** 
+	* Called when a project needs to be processed by the command. The output parameter is used to
+	* return output that will eventually be passed to ::ProcessOutput. If the command does not 
+	* produce any output needed by ::ProcessOutput then the parameter may be ignored.
+	*/
+	virtual bool ProcessProject(const FProject& Project, TUniquePtr<FCommandOutput>& Output) = 0;
+
+	/** 
+	 * Called after all projects have been processed. This will always be called by the original
+	 * processes and never called by a spawned child process. If the command produces output when
+	 * processing a project then the CmdOutputArray will be populated, otherwise it will be empty.
+	 * It is assumed that each command knows how to cast the FCommandOutput to the correct type.
+	 */
+	virtual bool ProcessOutput(const TArray<TUniquePtr<FCommandOutput>>& CmdOutputArray) = 0;
+
+	/** 
+	 * Derived command classes should return a new instance of the command output that it can produce.
+	 * This will be called when we are parsing child process output files while attempting to reconstruct
+	 * the output of a child process.
+	 * Commands that do not return any output can return nullptr. 
+	 */
+	virtual TUniquePtr<FCommandOutput> CreateOutputObject() const = 0;
 
 	virtual const TArray<FString>& GetPackages() const = 0;
 
@@ -72,6 +139,7 @@ protected: // Common SourceControl Code
 	bool TryParseChangelist(FStringView ClientSpecName, FStringView ChangelistNumber, TArray<FString>& OutPackages, FSourceControlChangelistPtr* OutChangelist);
 
 	FString FindClientSpecForChangelist(FStringView ChangelistNumber);
+
 private:
 	FString CommandName;
 
diff --git a/Engine/Source/Programs/UnrealVirtualizationTool/Private/Commands/RehydrateCommand.cpp b/Engine/Source/Programs/UnrealVirtualizationTool/Private/Commands/RehydrateCommand.cpp
--- a/Engine/Source/Programs/UnrealVirtualizationTool/Private/Commands/RehydrateCommand.cpp
+++ b/Engine/Source/Programs/UnrealVirtualizationTool/Private/Commands/RehydrateCommand.cpp
@@ -72,73 +72,89 @@ bool FRehydrateCommand::Initialize(const TCHAR* CmdLine)
 	return true;
 }
 
-bool FRehydrateCommand::Run(const TArray<FProject>& Projects)
+void FRehydrateCommand::Serialize(FJsonSerializerBase& Serializer)
 {
-	TRACE_CPUPROFILER_EVENT_SCOPE(FRehydrateCommand::Run);
+	TRACE_CPUPROFILER_EVENT_SCOPE(FRehydrateCommand::Serialize);
 
-	for (const FProject& Project : Projects)
-	{
-		TStringBuilder<128> ProjectName;
-		ProjectName << Project.GetProjectName();
+	// The command has no additional data when launched as a child process
+}
 
-		UE_LOG(LogVirtualizationTool, Display, TEXT("\tProcessing package(s) for the project '%s'..."), ProjectName.ToString());
+bool FRehydrateCommand::ProcessProject(const FProject& Project, TUniquePtr<FCommandOutput>& Output)
+{
+	TRACE_CPUPROFILER_EVENT_SCOPE(FRehydrateCommand::ProcessProject);
 
-		FConfigFile EngineConfigWithProject;
-		if (!Project.TryLoadConfig(EngineConfigWithProject))
-		{
-			return false;
-		}
+	TStringBuilder<128> ProjectName;
+	ProjectName << Project.GetProjectName();
 
-		Project.RegisterMountPoints();
+	UE_LOG(LogVirtualizationTool, Display, TEXT("\tProcessing package(s) for the project '%s'..."), ProjectName.ToString());
 
-		UE::Virtualization::FInitParams InitParams(ProjectName, EngineConfigWithProject);
-		UE::Virtualization::Initialize(InitParams, UE::Virtualization::EInitializationFlags::ForceInitialize);
+	FConfigFile EngineConfigWithProject;
+	if (!Project.TryLoadConfig(EngineConfigWithProject))
+	{
+		return false;
+	}
 
-		ON_SCOPE_EXIT
-		{
-			UE::Virtualization::Shutdown();
-		};
+	Project.RegisterMountPoints();
 
-		const TArray<FString> ProjectPackages = Project.GetAllPackages();
+	UE::Virtualization::FInitParams InitParams(ProjectName, EngineConfigWithProject);
+	UE::Virtualization::Initialize(InitParams, UE::Virtualization::EInitializationFlags::ForceInitialize);
 
-		if (!TryCheckOutFilesForProject(ClientSpecName, Project.GetProjectRoot(), ProjectPackages))
-		{
-			return false;
-		}
+	ON_SCOPE_EXIT
+	{
+		UE::Virtualization::Shutdown();
+	};
 
-		UE_LOG(LogVirtualizationTool, Display, TEXT("\t\tAttempting to rehydrate packages..."), ProjectName.ToString());
+	const TArray<FString> ProjectPackages = Project.GetAllPackages();
 
+	if (!TryCheckOutFilesForProject(ClientSpecName, Project.GetProjectRoot(), ProjectPackages))
+	{
+		return false;
+	}
 
-		ERehydrationOptions Options = ERehydrationOptions::None;
-		if (bShouldCheckout)
-		{
-			Options |= ERehydrationOptions::Checkout;
-		}
+	UE_LOG(LogVirtualizationTool, Display, TEXT("\t\tAttempting to rehydrate packages..."), ProjectName.ToString());
 
-		FRehydrationResult Result = UE::Virtualization::IVirtualizationSystem::Get().TryRehydratePackages(ProjectPackages, Options);
-		if (!Result.WasSuccessful())
-		{
-			UE_LOG(LogVirtualizationTool, Error, TEXT("The rehydration process failed with the following errors:"));
-			for (const FText& Error : Result.Errors)
-			{
-				UE_LOG(LogVirtualizationTool, Error, TEXT("\t%s"), *Error.ToString());
-			}
-			return false;
-		}
 
-		if (bShouldCheckout)
+	ERehydrationOptions Options = ERehydrationOptions::None;
+	if (bShouldCheckout)
+	{
+		Options |= ERehydrationOptions::Checkout;
+	}
+
+	FRehydrationResult Result = UE::Virtualization::IVirtualizationSystem::Get().TryRehydratePackages(ProjectPackages, Options);
+	if (!Result.WasSuccessful())
+	{
+		UE_LOG(LogVirtualizationTool, Error, TEXT("The rehydration process failed with the following errors:"));
+		for (const FText& Error : Result.Errors)
 		{
-			UE_LOG(LogVirtualizationTool, Display, TEXT("\t\t%d packages were checked out of revision control"), Result.CheckedOutPackages.Num());
+			UE_LOG(LogVirtualizationTool, Error, TEXT("\t%s"), *Error.ToString());
 		}
+		return false;
+	}
 
-		Project.UnRegisterMountPoints();
-
-		UE_LOG(LogVirtualizationTool, Display, TEXT("\t\tRehyration of project packages complete!"), ProjectName.ToString());
+	if (bShouldCheckout)
+	{
+		UE_LOG(LogVirtualizationTool, Display, TEXT("\t\t%d packages were checked out of revision control"), Result.CheckedOutPackages.Num());
 	}
 
+	Project.UnRegisterMountPoints();
+
+	UE_LOG(LogVirtualizationTool, Display, TEXT("\t\tRehyration of project packages complete!"), ProjectName.ToString());
+
 	return true;
 }
 
+bool FRehydrateCommand::ProcessOutput(const TArray<TUniquePtr<FCommandOutput>>& OutputArray)
+{
+	// Command has no additional work to do
+	return true;
+}
+
+TUniquePtr<FCommandOutput> FRehydrateCommand::CreateOutputObject() const
+{
+	// This command does not create any output
+	return TUniquePtr<FCommandOutput>();
+}
+
 const TArray<FString>& FRehydrateCommand::GetPackages() const
 {
 	return Packages;
diff --git a/Engine/Source/Programs/UnrealVirtualizationTool/Private/Commands/RehydrateCommand.h b/Engine/Source/Programs/UnrealVirtualizationTool/Private/Commands/RehydrateCommand.h
--- a/Engine/Source/Programs/UnrealVirtualizationTool/Private/Commands/RehydrateCommand.h
+++ b/Engine/Source/Programs/UnrealVirtualizationTool/Private/Commands/RehydrateCommand.h
@@ -17,7 +17,12 @@ public:
 private:
 	virtual bool Initialize(const TCHAR* CmdLine) override;
 
-	virtual bool Run(const TArray<FProject>& Projects) override;
+	virtual void Serialize(FJsonSerializerBase& Serializer) override;
+
+	virtual bool ProcessProject(const FProject& Project, TUniquePtr<FCommandOutput>& Output) override;
+	virtual bool ProcessOutput(const TArray<TUniquePtr<FCommandOutput>>& OutputArray) override;
+
+	virtual TUniquePtr<FCommandOutput> CreateOutputObject() const override;
 
 	virtual const TArray<FString>& GetPackages() const override;
 
diff --git a/Engine/Source/Programs/UnrealVirtualizationTool/Private/Commands/VirtualizeCommand.cpp b/Engine/Source/Programs/UnrealVirtualizationTool/Private/Commands/VirtualizeCommand.cpp
--- a/Engine/Source/Programs/UnrealVirtualizationTool/Private/Commands/VirtualizeCommand.cpp
+++ b/Engine/Source/Programs/UnrealVirtualizationTool/Private/Commands/VirtualizeCommand.cpp
@@ -14,27 +14,40 @@
 #include "ISourceControlProvider.h"
 #include "SourceControlOperations.h"
 
-namespace
+namespace UE::Virtualization
 {
 
 /** Utility */
-TArray<FString> BuildFinalTagDescriptions(TArray<FText>& DescriptionTags)
+TArray<FString> BuildFinalTagDescriptions(const TArray<TUniquePtr<FCommandOutput>>& OutputArray)
 {
 	TArray<FString> CleanedDescriptions;
-	CleanedDescriptions.Reserve(DescriptionTags.Num());
-
-	for (const FText& Tag : DescriptionTags)
+	
+	for (const TUniquePtr<FCommandOutput>& Output : OutputArray)
 	{
-		CleanedDescriptions.AddUnique(Tag.ToString());
+		if (Output)
+		{
+			const FVirtualizeCommandOutput* CmdOutput = (const FVirtualizeCommandOutput*)Output.Get();
+
+			for (const FString& Tag : CmdOutput->DescriptionTags)
+			{
+				CleanedDescriptions.AddUnique(Tag);
+			}
+		}
 	}
 
 	return CleanedDescriptions;
 }
 
-} // namespace
-
-namespace UE::Virtualization
+FVirtualizeCommandOutput::FVirtualizeCommandOutput(FStringView InProjectName, const TArray<FText>& InDescriptionTags)
+	: FCommandOutput(InProjectName)
 {
+	DescriptionTags.Reserve(InDescriptionTags.Num());
+
+	for (const FText& Description : InDescriptionTags)
+	{
+		DescriptionTags.Add(Description.ToString());
+	}
+}
 
 FVirtualizeCommand::FVirtualizeCommand(FStringView CommandName)
 	: FCommand(CommandName)
@@ -122,73 +135,84 @@ bool FVirtualizeCommand::Initialize(const TCHAR* CmdLine)
 	return true;
 }
 
-bool FVirtualizeCommand::Run(const TArray<FProject>& Projects)
+void FVirtualizeCommand::Serialize(FJsonSerializerBase& Serializer)
 {
-	TRACE_CPUPROFILER_EVENT_SCOPE(FVirtualizeCommand::Run);
+	TRACE_CPUPROFILER_EVENT_SCOPE(FVirtualizeCommand::Serialize);
+
+	Serializer.Serialize(TEXT("ShouldCheckout"), bShouldCheckout);
+}
+
+bool FVirtualizeCommand::ProcessProject(const FProject& Project, TUniquePtr<FCommandOutput>& Output)
+{
+	TRACE_CPUPROFILER_EVENT_SCOPE(FVirtualizeCommand::ProcessProject);
 
 	UE_LOG(LogVirtualizationTool, Display, TEXT("Running the virtualization process..."));
-	
-	TArray<FText> DescriptionTags;
 
-	for (const FProject& Project : Projects)
+	TStringBuilder<128> ProjectName;
+	ProjectName << Project.GetProjectName();
+
+	UE_LOG(LogVirtualizationTool, Display, TEXT("\tChecking package(s) for the project '%s'..."), ProjectName.ToString());
+
+	FConfigFile EngineConfigWithProject;
+	if (!Project.TryLoadConfig(EngineConfigWithProject))
 	{
-		TStringBuilder<128> ProjectName;
-		ProjectName << Project.GetProjectName();
+		return false;
+	}
 
-		UE_LOG(LogVirtualizationTool, Display, TEXT("\tChecking package(s) for the project '%s'..."), ProjectName.ToString());
+	Project.RegisterMountPoints();
 
-		FConfigFile EngineConfigWithProject;
-		if (!Project.TryLoadConfig(EngineConfigWithProject))
-		{
-			return false;
-		}
+	ON_SCOPE_EXIT
+	{
+		Project.UnRegisterMountPoints();
+	};
 
-		Project.RegisterMountPoints();
+	UE::Virtualization::FInitParams InitParams(ProjectName, EngineConfigWithProject);
+	UE::Virtualization::Initialize(InitParams, UE::Virtualization::EInitializationFlags::ForceInitialize);
 
-		ON_SCOPE_EXIT
-		{
-			Project.UnRegisterMountPoints();
-		};
+	ON_SCOPE_EXIT
+	{
+		UE::Virtualization::Shutdown();
+	};
 
-		UE::Virtualization::FInitParams InitParams(ProjectName, EngineConfigWithProject);
-		UE::Virtualization::Initialize(InitParams, UE::Virtualization::EInitializationFlags::ForceInitialize);
+	TArray<FString> ProjectPackages = Project.GetAllPackages();
 
-		ON_SCOPE_EXIT
-		{
-			UE::Virtualization::Shutdown();
-		};
+	EVirtualizationOptions Options = EVirtualizationOptions::None;
+	if (bShouldCheckout)
+	{
+		Options |= EVirtualizationOptions::Checkout;
+	}
 
-		TArray<FString> ProjectPackages = Project.GetAllPackages();
+	FVirtualizationResult Result = UE::Virtualization::IVirtualizationSystem::Get().TryVirtualizePackages(ProjectPackages, Options);
 
-		EVirtualizationOptions Options = EVirtualizationOptions::None;
-		if(bShouldCheckout)
+	if (!Result.WasSuccessful())
+	{
+		UE_LOG(LogVirtualizationTool, Error, TEXT("The virtualization process failed with the following errors:"));
+		for (const FText& Error : Result.Errors)
 		{
-			Options |= EVirtualizationOptions::Checkout;
+			UE_LOG(LogVirtualizationTool, Error, TEXT("\t%s"), *Error.ToString());
 		}
+		return false;
+	}
 
-		FVirtualizationResult Result = UE::Virtualization::IVirtualizationSystem::Get().TryVirtualizePackages(ProjectPackages, Options);
+	if (bShouldCheckout)
+	{
+		UE_LOG(LogVirtualizationTool, Display, TEXT("\t\t%d packages were checked out of revision control"), Result.CheckedOutPackages.Num());
+	}
 
-		if (!Result.WasSuccessful())
-		{
-			UE_LOG(LogVirtualizationTool, Error, TEXT("The virtualization process failed with the following errors:"));
-			for (const FText& Error : Result.Errors)
-			{
-				UE_LOG(LogVirtualizationTool, Error, TEXT("\t%s"), *Error.ToString());
-			}
-			return false;
-		}
+	UE_LOG(LogVirtualizationTool, Display, TEXT("\t\tVirtualization of project packages complete!"), ProjectName.ToString());
 
-		if (bShouldCheckout)
-		{
-			UE_LOG(LogVirtualizationTool, Display, TEXT("\t\t%d packages were checked out of revision control"), Result.CheckedOutPackages.Num());
-		}
+	Output = MakeUnique<FVirtualizeCommandOutput>(Project.GetProjectName(), Result.DescriptionTags);
 
-		UE_LOG(LogVirtualizationTool, Display, TEXT("\t\tVirtualization of project packages complete!"), ProjectName.ToString());
-	}
+	return true;
+}
+
+bool FVirtualizeCommand::ProcessOutput(const TArray<TUniquePtr<FCommandOutput>>& CmdOutputArray)
+{
+	TRACE_CPUPROFILER_EVENT_SCOPE(FVirtualizeCommand::ProcessOutput);
 
 	if (bShouldSubmitChangelist)
 	{
-		TArray<FString> FinalDescriptionTags = BuildFinalTagDescriptions(DescriptionTags);
+		TArray<FString> FinalDescriptionTags = BuildFinalTagDescriptions(CmdOutputArray);
 
 		if (!TrySubmitChangelist(SourceChangelist, FinalDescriptionTags))
 		{
@@ -199,12 +223,17 @@ bool FVirtualizeCommand::Run(const TArray<FProject>& Projects)
 	return true;
 }
 
+TUniquePtr<FCommandOutput> FVirtualizeCommand::CreateOutputObject() const
+{
+	return MakeUnique<FVirtualizeCommandOutput>();
+}
+
 const TArray<FString>& FVirtualizeCommand::GetPackages() const
 {
 	return AllPackages;
 }
 
-bool FVirtualizeCommand::TrySubmitChangelist(const FSourceControlChangelistPtr& ChangelistToSubmit, const TArray<FString>& DescriptionTags)
+bool FVirtualizeCommand::TrySubmitChangelist(const FSourceControlChangelistPtr& ChangelistToSubmit, const TArray<FString>& InDescriptionTags)
 {
 	if (!ChangelistToSubmit.IsValid())
 	{
@@ -238,7 +267,7 @@ bool FVirtualizeCommand::TrySubmitChangelist(const FSourceControlChangelistPtr&
 	TStringBuilder<512> Description;
 	Description << ChangelistState->GetDescriptionText().ToString();
 
-	for (const FString& Tag : DescriptionTags)
+	for (const FString& Tag : InDescriptionTags)
 	{
 		Description << TEXT("\n") << Tag;
 	}
diff --git a/Engine/Source/Programs/UnrealVirtualizationTool/Private/Commands/VirtualizeCommand.h b/Engine/Source/Programs/UnrealVirtualizationTool/Private/Commands/VirtualizeCommand.h
--- a/Engine/Source/Programs/UnrealVirtualizationTool/Private/Commands/VirtualizeCommand.h
+++ b/Engine/Source/Programs/UnrealVirtualizationTool/Private/Commands/VirtualizeCommand.h
@@ -7,6 +7,19 @@
 namespace UE::Virtualization
 {
 
+struct FVirtualizeCommandOutput final : public FCommandOutput
+{
+	FVirtualizeCommandOutput() = default;
+	FVirtualizeCommandOutput(FStringView InProjectName, const TArray<FText>& InDescriptionTags);
+
+	BEGIN_JSON_SERIALIZER
+		JSON_SERIALIZE_PARENT(FCommandOutput);
+		JSON_SERIALIZE_ARRAY("DescriptionTags", DescriptionTags);
+	END_JSON_SERIALIZER
+
+	TArray<FString> DescriptionTags;
+};
+
 class FVirtualizeCommand : public FCommand
 {
 public:
@@ -18,7 +31,12 @@ public:
 private:
 	virtual bool Initialize(const TCHAR* CmdLine) override;
 
-	virtual bool Run(const TArray<FProject>& Projects) override;
+	virtual void Serialize(FJsonSerializerBase& Serializer) override;
+
+	virtual bool ProcessProject(const FProject& Project, TUniquePtr<FCommandOutput>& Output) override;
+	virtual bool ProcessOutput(const TArray<TUniquePtr<FCommandOutput>>& CmdOutputArray) override;
+
+	virtual TUniquePtr<FCommandOutput> CreateOutputObject() const override;
 
 	virtual const TArray<FString>& GetPackages() const override;
 
@@ -62,6 +80,4 @@ private:
 	virtual bool Initialize(const TCHAR* CmdLine) override;
 };
 
-
-
 } // namespace UE::Virtualization
diff --git a/Engine/Source/Programs/UnrealVirtualizationTool/Private/ProcessUtilities.h b/Engine/Source/Programs/UnrealVirtualizationTool/Private/ProcessUtilities.h
new file mode 100644
--- /dev/null
+++ b/Engine/Source/Programs/UnrealVirtualizationTool/Private/ProcessUtilities.h
@@ -0,0 +1,62 @@
+// Copyright Epic Games, Inc. All Rights Reserved.
+
+#pragma once
+
+#include "Containers/Array.h"
+#include "Containers/UnrealString.h"
+#include "HAL/PlatformProcess.h"
+#include "Logging/LogMacros.h"
+#include "UnrealVirtualizationTool.h"
+
+namespace UE::Virtualization
+{
+
+/** Wrapper around stdin and stdout pipes created by a call to FPlatformProcess::CreatePipe */
+struct FProcessPipes
+{
+	FProcessPipes()
+	{
+		verify(FPlatformProcess::CreatePipe(StdOut, StdIn));
+	}
+
+	~FProcessPipes()
+	{
+		FPlatformProcess::ClosePipe(StdOut, StdIn);
+	}
+
+	void* GetStdIn() const
+	{
+		return StdIn;
+	}
+
+	void* GetStdOut() const
+	{
+		return StdOut;
+	}
+
+	void ProcessStdOut()
+	{
+		check(StdOut != nullptr);
+
+		FString Output = FPlatformProcess::ReadPipe(StdOut);
+
+		while (!Output.IsEmpty())
+		{
+			TArray<FString> Lines;
+			Output.ParseIntoArray(Lines, LINE_TERMINATOR);
+
+			for (const FString& Line : Lines)
+			{
+				UE_LOG(LogVirtualizationTool, Display, TEXT("Child Process-> %s"), *Line);
+			}
+
+			Output = FPlatformProcess::ReadPipe(StdOut);
+		}
+	}
+
+private:
+	void* StdIn = nullptr;
+	void* StdOut = nullptr;
+};
+
+} // namespace UE::Virtualization
\ No newline at end of file
diff --git a/Engine/Source/Programs/UnrealVirtualizationTool/Private/ProjectFiles.cpp b/Engine/Source/Programs/UnrealVirtualizationTool/Private/ProjectFiles.cpp
--- a/Engine/Source/Programs/UnrealVirtualizationTool/Private/ProjectFiles.cpp
+++ b/Engine/Source/Programs/UnrealVirtualizationTool/Private/ProjectFiles.cpp
@@ -133,7 +133,7 @@ void FProject::AddPluginFile(const FString& PackagePath, FString&& PluginFilePat
 	Plugin->PackagePaths.Add(PackagePath);
 };
 
-FStringView FProject::GetProjectFilePath() const
+const FString& FProject::GetProjectFilePath() const
 {
 	return ProjectFilePath;
 }
diff --git a/Engine/Source/Programs/UnrealVirtualizationTool/Private/ProjectFiles.h b/Engine/Source/Programs/UnrealVirtualizationTool/Private/ProjectFiles.h
--- a/Engine/Source/Programs/UnrealVirtualizationTool/Private/ProjectFiles.h
+++ b/Engine/Source/Programs/UnrealVirtualizationTool/Private/ProjectFiles.h
@@ -5,27 +5,36 @@
 #include "Containers/Array.h"
 #include "Containers/StringFwd.h"
 #include "Containers/UnrealString.h"
+#include "Serialization/JsonSerializerMacros.h"
 
 namespace UE::Virtualization
 {
 
-struct FPlugin
+struct FPlugin : public FJsonSerializable
 {
+public:
 	FString PluginFilePath;
 
 	TArray<FString> PackagePaths;
+
+public:
+	BEGIN_JSON_SERIALIZER
+		JSON_SERIALIZE("PluginPath", PluginFilePath);
+		JSON_SERIALIZE_ARRAY("PackagePaths", PackagePaths);
+	END_JSON_SERIALIZER
 };
 
-class FProject
+class FProject : public FJsonSerializable
 {
 public:
+	FProject() = default;
 	FProject(FString&& ProjectFilePath);
 	~FProject() = default;
 
 	void AddFile(const FString& FilePath);
 	void AddPluginFile(const FString& FilePath, FString&& PluginFilePath);
 	
-	FStringView GetProjectFilePath() const;
+	const FString& GetProjectFilePath() const;
 	FStringView GetProjectName() const;
 	FStringView GetProjectRoot() const;
 
@@ -41,6 +50,13 @@ private:
 
 	TArray<FPlugin> Plugins;
 	TArray<FString> PackagePaths;
+
+public:
+	BEGIN_JSON_SERIALIZER
+		JSON_SERIALIZE("ProjectPath", ProjectFilePath);
+		JSON_SERIALIZE_ARRAY_SERIALIZABLE("Plugins", Plugins, FPlugin);
+		JSON_SERIALIZE_ARRAY("PackagePaths", PackagePaths);
+	END_JSON_SERIALIZER
 };
 
-} //namespace UE::Virtualization
\ No newline at end of file
+} //namespace UE::Virtualization
diff --git a/Engine/Source/Programs/UnrealVirtualizationTool/Private/UnrealVirtualizationToolApp.cpp b/Engine/Source/Programs/UnrealVirtualizationTool/Private/UnrealVirtualizationToolApp.cpp
--- a/Engine/Source/Programs/UnrealVirtualizationTool/Private/UnrealVirtualizationToolApp.cpp
+++ b/Engine/Source/Programs/UnrealVirtualizationTool/Private/UnrealVirtualizationToolApp.cpp
@@ -8,6 +8,7 @@
 #include "GenericPlatform/GenericPlatformOutputDevices.h"
 #include "HAL/FeedbackContextAnsi.h"
 #include "HAL/FileManager.h"
+#include "HAL/PlatformProcess.h"
 #include "Interfaces/IPluginManager.h"
 #include "Logging/StructuredLog.h"
 #include "Misc/CommandLine.h"
@@ -18,6 +19,7 @@
 #include "Misc/Paths.h"
 #include "Misc/PathViews.h"
 #include "Modules/ModuleManager.h"
+#include "ProcessUtilities.h"
 #include "UnrealVirtualizationTool.h"
 #include "Virtualization/VirtualizationSystem.h"
 
@@ -49,6 +51,30 @@ void LexFromString(UE::Virtualization::EMode& OutValue, const FStringView& InStr
 	}
 }
 
+/** 
+ * Utility to convert EMode to a string. NOTE that legacy modes will return the
+ * name of the newer mode that replaces it and not the legacy mode.
+ */
+TCHAR* LexToString(UE::Virtualization::EMode Mode)
+{
+	using namespace UE::Virtualization;
+
+	switch (Mode)
+	{
+		case EMode::Unknown:
+			return TEXT("Unknown");
+		case EMode::Changelist:		// Legacy Mode
+		case EMode::PackageList:	// Legacy Mode
+		case EMode::Virtualize:
+			return TEXT("Virtualize");
+		case EMode::Rehydrate:
+			return TEXT("Rehydrate");	
+		default:
+			checkNoEntry();
+			return TEXT("");
+	}
+}
+
 /** Utility for creating a new command */
 template<typename CommandType>
 TUniquePtr<UE::Virtualization::FCommand> CreateCommand(const FString& ModeName, const TCHAR* CmdLine)
@@ -66,6 +92,20 @@ TUniquePtr<UE::Virtualization::FCommand> CreateCommand(const FString& ModeName,
 	}
 }
 
+/** Utility to create a file path for the child process input/output file */
+void CreateChildProcessFilePath(FStringView Id, FStringView Extension, FStringBuilderBase& OutPath)
+{
+	FPathViews::ToAbsolutePath(FPaths::EngineSavedDir(), OutPath);
+	OutPath << TEXT("UnrealVirtualizationTool/") << Id << TEXT(".") << Extension;
+}
+
+/** Utility to create a file path for the child process input/output file from a FGuid */
+void CreateChildProcessFilePath(const FGuid& Id, FStringView Extension, FStringBuilderBase& OutPath)
+{
+	CreateChildProcessFilePath(WriteToString<40>(Id), Extension, OutPath);
+}
+
+
 /**
  * This class can be used to prevent log messages from other systems being logged with the Display verbosity.
  * In practical terms this means as long as the class is alive, only LogVirtualizationTool messages will
@@ -160,13 +200,16 @@ EInitResult FUnrealVirtualizationToolApp::Initialize()
 		return CmdLineResult;
 	}
 
-	TArray<FString> Packages = CurrentCommand->GetPackages();
+	if (!IsChildProcess())
+	{
+		TArray<FString> Packages = CurrentCommand->GetPackages();
 
-	UE_LOG(LogVirtualizationTool, Display, TEXT("\tFound '%d' package file(s)"), Packages.Num());
+		UE_LOG(LogVirtualizationTool, Display, TEXT("\tFound '%d' package file(s)"), Packages.Num());
 
-	if (!TrySortFilesByProject(Packages))
-	{
-		return EInitResult::Error;
+		if (!TrySortFilesByProject(Packages))
+		{
+			return EInitResult::Error;
+		}
 	}
 
 	UE_LOG(LogVirtualizationTool, Display, TEXT("Initialization complete!"));
@@ -178,18 +221,82 @@ bool FUnrealVirtualizationToolApp::Run()
 {
 	TRACE_CPUPROFILER_EVENT_SCOPE(Run);
 
-	UE_LOG(LogVirtualizationTool, Display, TEXT("Running the '%s' command..."), *CurrentCommand->GetName());
+	TArray<TUniquePtr<FCommandOutput>> OutputArray;
+	if (!ProcessProjects(OutputArray))
+	{
+		UE_LOG(LogVirtualizationTool, Error, TEXT("Command '%s' failed!"), *CurrentCommand->GetName());
+		return false;
+	}
 
-	if (CurrentCommand->Run(Projects))
+	if (!IsChildProcess())
 	{
-		UE_LOG(LogVirtualizationTool, Display, TEXT("Command '%s' succeeded!"), *CurrentCommand->GetName());
-		return true;
+		if (!CurrentCommand->ProcessOutput(OutputArray))
+		{
+			UE_LOG(LogVirtualizationTool, Error, TEXT("Command '%s' failed!"), *CurrentCommand->GetName());
+			return false;
+		}
 	}
 	else
 	{
-		UE_LOG(LogVirtualizationTool, Error, TEXT("Command '%s' failed!"), *CurrentCommand->GetName());
-		return false;
+		if (!TryWriteChildProcessOutputFile(ChildProcessId, OutputArray))
+		{
+			UE_LOG(LogVirtualizationTool, Error, TEXT("Command '%s' failed!"), *CurrentCommand->GetName());
+			return false;
+		}
 	}
+
+
+	UE_LOG(LogVirtualizationTool, Display, TEXT("Command '%s' succeeded!"), *CurrentCommand->GetName());
+	return true;
+}
+
+bool FUnrealVirtualizationToolApp::ProcessProjects(TArray<TUniquePtr<FCommandOutput>>& OutputArray)
+{
+	TRACE_CPUPROFILER_EVENT_SCOPE(ProcessProjects);
+
+	UE_LOG(LogVirtualizationTool, Display, TEXT("Running the '%s' command..."), *CurrentCommand->GetName());
+
+	OutputArray.Reserve(Projects.Num());
+
+	for (const FProject& Project : Projects)
+	{
+		UE_LOG(LogVirtualizationTool, Display, TEXT("Processing project: %s"), *Project.GetProjectFilePath());
+
+		const FString ProjectPath = FPaths::ConvertRelativePathToFull(FPaths::GetProjectFilePath());
+
+		if (IsChildProcess() || ProjectPath == Project.GetProjectFilePath())
+		{
+			if (IsChildProcess() && ProjectPath != Project.GetProjectFilePath())
+			{
+				// If we are a child process then the correct project path should have been provided and so
+				// this check is mostly a paranoid check to make sure things are working as we expect.
+				UE_LOG(LogVirtualizationTool, Error, TEXT("The child process was create with project path '%s' but expected '%s'"), *ProjectPath, *Project.GetProjectFilePath());
+				return false;
+			}
+
+			check(ProjectPath == Project.GetProjectFilePath());
+
+			TUniquePtr<FCommandOutput> Output;
+			if (!CurrentCommand->ProcessProject(Project, Output))
+			{
+				return false;
+			}
+
+			if (Output != nullptr)
+			{
+				OutputArray.Emplace(MoveTemp(Output));
+			}
+		}
+		else
+		{
+			if (!LaunchChildProcess(*CurrentCommand, Project, GlobalCmdlineOptions, OutputArray))
+			{
+				return false;
+			}
+		}
+	}
+
+	return true;
 }
 
 void FUnrealVirtualizationToolApp::PrintCmdLineHelp() const
@@ -270,14 +377,24 @@ EInitResult FUnrealVirtualizationToolApp::TryParseCmdLine()
 		return EInitResult::EarlyOut;
 	}
 
-	if (FParse::Param(CmdLine, TEXT("Verbose")))
+	EInitResult GlobalOptionResult = TryParseGlobalOptions(CmdLine);
+	if(GlobalOptionResult != EInitResult::Success)
 	{
-		UE_LOG(LogVirtualizationTool, Display, TEXT("Cmdline parameter '-Verbose' found, no longer supressing Display log messages!"));
-		OutputDeviceOverride.Reset();
+		return GlobalOptionResult;
 	}
-	else
+
+	// Check to see if we are a child process with an input file
+	FString ChildProcessInput;
+	if (FParse::Value(CmdLine, TEXT("-ChildProcess="), ChildProcessInput))
 	{
-		OutputDeviceOverride = MakeUnique<FOverrideOutputDevice>();
+		if (TryReadChildProcessInputFile(ChildProcessInput))
+		{
+			return EInitResult::Success;
+		}
+		else
+		{
+			return EInitResult::Error;
+		}
 	}
 
 	// Now parse the mode specific command line options
@@ -290,24 +407,50 @@ EInitResult FUnrealVirtualizationToolApp::TryParseCmdLine()
 		return EInitResult::Error;
 	}
 
-	LexFromString(Mode, ModeAsString);
+	return CreateCommandFromString(ModeAsString, CmdLine);
+}
+
+EInitResult FUnrealVirtualizationToolApp::TryParseGlobalOptions(const TCHAR* CmdLine)
+{
+	GlobalCmdlineOptions.Reset();
+
+	if (FParse::Param(CmdLine, TEXT("verbose")))
+	{
+		UE_LOG(LogVirtualizationTool, Display, TEXT("Cmdline parameter '-verbose' found, no longer supressing Display log messages!"));
+		OutputDeviceOverride.Reset();
+
+		AddGlobalOption(TEXT("-verbose"));	
+	}
+	else
+	{
+		OutputDeviceOverride = MakeUnique<FOverrideOutputDevice>();
+	}
+
+	return EInitResult::Success;
+}
+
+EInitResult FUnrealVirtualizationToolApp::CreateCommandFromString(const FString& CommandName, const TCHAR* Cmdline)
+{
+	check(Mode == EMode::Unknown && !CurrentCommand.IsValid());
+
+	LexFromString(Mode, CommandName);
 
 	switch (Mode)
 	{
 		case EMode::Changelist:
-			CurrentCommand = CreateCommand<FVirtualizeLegacyChangeListCommand>(ModeAsString, CmdLine);
+			CurrentCommand = CreateCommand<FVirtualizeLegacyChangeListCommand>(LexToString(Mode), Cmdline);
 			return CurrentCommand.IsValid() ? EInitResult::Success : EInitResult::Error;
 			break;
 		case EMode::PackageList:
-			CurrentCommand = CreateCommand<FVirtualizeLegacyPackageListCommand>(ModeAsString, CmdLine);
+			CurrentCommand = CreateCommand<FVirtualizeLegacyPackageListCommand>(LexToString(Mode), Cmdline);
 			return CurrentCommand.IsValid() ? EInitResult::Success : EInitResult::Error;
 			break;
 		case EMode::Virtualize:
-			CurrentCommand = CreateCommand<FVirtualizeCommand>(ModeAsString, CmdLine);
+			CurrentCommand = CreateCommand<FVirtualizeCommand>(CommandName, Cmdline);
 			return CurrentCommand.IsValid() ? EInitResult::Success : EInitResult::Error;
 			break;
 		case EMode::Rehydrate:
-			CurrentCommand = CreateCommand<FRehydrateCommand>(ModeAsString, CmdLine);
+			CurrentCommand = CreateCommand<FRehydrateCommand>(CommandName, Cmdline);
 			return CurrentCommand.IsValid() ? EInitResult::Success : EInitResult::Error;
 			break;
 		case EMode::Unknown:
@@ -440,4 +583,292 @@ FProject& FUnrealVirtualizationToolApp::FindOrAddProject(FString&& ProjectFilePa
 	}
 }
 
+bool FUnrealVirtualizationToolApp::IsChildProcess() const
+{
+	return !ChildProcessId.IsEmpty();
+}
+
+void FUnrealVirtualizationToolApp::AddGlobalOption(FStringView Options)
+{
+	if (!GlobalCmdlineOptions.IsEmpty())
+	{
+		GlobalCmdlineOptions.Append(TEXT(" "));
+	}
+
+	GlobalCmdlineOptions.Append(Options);
+}
+
+bool FUnrealVirtualizationToolApp::TryReadChildProcessOutputFile(const FGuid& ChildProcessId, const FCommand& Command, TArray<TUniquePtr<FCommandOutput>>& OutputArray)
+{
+	TRACE_CPUPROFILER_EVENT_SCOPE(TryReadChildProcessOutputFile);
+
+	TStringBuilder<512> FilePath;
+	CreateChildProcessFilePath(ChildProcessId, TEXT("output"), FilePath);
+
+	FString JsonText;
+	if (!FFileHelper::LoadFileToString(JsonText, FilePath.ToString()))
+	{
+		UE_LOG(LogVirtualizationTool, Error, TEXT("Could not open child process output file '%s'"), FilePath.ToString());
+		return false;
+	}
+
+	TSharedPtr<FJsonObject> JsonRootObject;
+	TSharedRef<TJsonReader<TCHAR>> Reader = TJsonReaderFactory<TCHAR>::Create(JsonText);
+	if (!FJsonSerializer::Deserialize(Reader, JsonRootObject) || !JsonRootObject.IsValid())
+	{
+		UE_LOG(LogVirtualizationTool, Error, TEXT("Failed to parse child process output file '%s'"), FilePath.ToString());
+		return false;
+	}
+
+	const TArray<TSharedPtr<FJsonValue>>* OutputJsonArray = nullptr;
+	if (JsonRootObject->TryGetArrayField(TEXT("OutputArray"), OutputJsonArray))
+	{
+		check(OutputJsonArray != nullptr);
+
+		for (TSharedPtr<FJsonValue> OutputValue : *OutputJsonArray)
+		{
+			TUniquePtr<FCommandOutput> Output = Command.CreateOutputObject();
+
+			const TSharedPtr<FJsonObject>* OutputObject = nullptr;
+
+			if (!OutputValue->TryGetObject(OutputObject))
+			{
+				UE_LOG(LogVirtualizationTool, Error, TEXT("Invalid syntax found in child process output file '%s'"), FilePath.ToString());
+				return false;
+			}
+
+			check(OutputObject != nullptr);
+			if (!Output->FromJson(*OutputObject))
+			{
+				UE_LOG(LogVirtualizationTool, Error, TEXT("Failed to read FCommandOutput from the child process output file '%s'"), FilePath.ToString());
+				return false;
+			}
+
+			OutputArray.Emplace(MoveTemp(Output));
+		}
+	}
+	else
+	{
+		UE_LOG(LogVirtualizationTool, Error, TEXT("Invalid syntax found in child process output file '%s'"), FilePath.ToString());
+		return false;
+	}
+
+	return true;
+}
+
+bool FUnrealVirtualizationToolApp::TryWriteChildProcessOutputFile(const FString& ChildProcessId, const TArray<TUniquePtr<FCommandOutput>>& OutputArray)
+{
+	TRACE_CPUPROFILER_EVENT_SCOPE(TryWriteChildProcessOutputFile);
+
+	FString JsonTcharText;
+	TSharedRef<TJsonWriter<TCHAR, TPrettyJsonPrintPolicy<TCHAR>>> Writer = TJsonWriterFactory<TCHAR, TPrettyJsonPrintPolicy<TCHAR>>::Create(&JsonTcharText);
+
+	Writer->WriteObjectStart();
+	Writer->WriteArrayStart(TEXT("OutputArray"));
+	for (const TUniquePtr<FCommandOutput>& Output : OutputArray)
+	{
+		if (Output)
+		{
+			Writer->WriteObjectStart();
+			Output->ToJson(Writer, true);
+			Writer->WriteObjectEnd();
+		}
+	}
+	Writer->WriteArrayEnd();
+	Writer->WriteObjectEnd();
+
+	if (!Writer->Close())
+	{
+		UE_LOG(LogVirtualizationTool, Display, TEXT("Failed to create child process output file json document"));
+		return false;
+	}
+
+	TStringBuilder<512> FilePath;
+	CreateChildProcessFilePath(ChildProcessId, TEXT("output"), FilePath);
+
+	if (!FFileHelper::SaveStringToFile(JsonTcharText, FilePath.ToString()))
+	{
+		UE_LOG(LogVirtualizationTool, Display, TEXT("Failed to create child process output file '%s'"), FilePath.ToString());
+		return false;
+	}
+
+	return true;
+}
+
+bool FUnrealVirtualizationToolApp::TryReadChildProcessInputFile(const FString& InputPath)
+{
+	UE_LOG(LogVirtualizationTool, Display, TEXT("Parsing child process input file..."));
+
+	FString JsonText;
+	if (!FFileHelper::LoadFileToString(JsonText, *InputPath))
+	{
+		UE_LOG(LogVirtualizationTool, Error, TEXT("Could not open child process input file '%s'"), *InputPath);
+		return false;
+	}
+
+	TSharedPtr<FJsonObject> JsonRootObject;
+	TSharedRef<TJsonReader<TCHAR>> Reader = TJsonReaderFactory<TCHAR>::Create(JsonText);
+	if (!FJsonSerializer::Deserialize(Reader, JsonRootObject) || !JsonRootObject.IsValid())
+	{
+		UE_LOG(LogVirtualizationTool, Error, TEXT("Failed to parse child process input file '%s'"), *InputPath);
+		return false;
+	}
+
+	if (TSharedPtr<FJsonValue> JsonValue = JsonRootObject->TryGetField(TEXT("CommandName")))
+	{
+		const FString CommandName = JsonValue->AsString();
+
+		if (CreateCommandFromString(CommandName, TEXT("")) != EInitResult::Success)
+		{
+			UE_LOG(LogVirtualizationTool, Error, TEXT("Failed to create command '%s'"), *CommandName);
+			return false;
+		}
+	}
+	else
+	{
+		UE_LOG(LogVirtualizationTool, Error, TEXT("Failed to find 'CommandName' in child process input file '%s'"), *InputPath);
+		return false;
+	}
+
+	if (TSharedPtr<FJsonObject> JsonObject = JsonRootObject->GetObjectField(TEXT("ProjectData")))
+	{
+		FProject Project;
+		if (Project.FromJson(JsonObject))
+		{
+			Projects.Emplace(MoveTemp(Project));
+		}
+		else
+		{
+			UE_LOG(LogVirtualizationTool, Error, TEXT("Failed to serialize project data from child process input file '%s'"), *InputPath);
+			return false;
+		}
+	}
+	else
+	{
+		UE_LOG(LogVirtualizationTool, Error, TEXT("Failed to find 'ProjectData' in child process input file '%s'"), *InputPath);
+		return false;
+	}
+
+	if (TSharedPtr<FJsonObject> JsonObject = JsonRootObject->GetObjectField(TEXT("CommandData")))
+	{
+		CurrentCommand->FromJson(JsonObject);
+	}
+	else
+	{
+		UE_LOG(LogVirtualizationTool, Error, TEXT("Failed to find 'ProjectData' in child process input file '%s'"), *InputPath);
+		return false;
+	}
+
+	ChildProcessId = FPathViews::GetBaseFilename(InputPath);
+
+	return true;
+}
+
+bool FUnrealVirtualizationToolApp::TryWriteChildProcessInputFile(const FGuid& ChildProcessId, const FCommand& Command, const FProject& Project, FStringBuilderBase& OutPath)
+{
+	TRACE_CPUPROFILER_EVENT_SCOPE(TryWriteChildProcessInputFile);
+
+	FString JsonTcharText;
+	TSharedRef<TJsonWriter<TCHAR, TPrettyJsonPrintPolicy<TCHAR>>> Writer = TJsonWriterFactory<TCHAR, TPrettyJsonPrintPolicy<TCHAR>>::Create(&JsonTcharText);
+	Writer->WriteObjectStart();
+
+	Writer->WriteValue(TEXT("CommandName"), Command.GetName());
+
+	Writer->WriteObjectStart(TEXT("ProjectData"));
+	Project.ToJson(Writer, true);
+	Writer->WriteObjectEnd();
+
+	Writer->WriteObjectStart(TEXT("CommandData"));
+	Command.ToJson(Writer);
+	Writer->WriteObjectEnd();
+
+	Writer->WriteObjectEnd();
+
+	if (!Writer->Close())
+	{
+		UE_LOG(LogVirtualizationTool, Display, TEXT("Failed to create child process input file json document"));
+		return false;
+	}
+
+	CreateChildProcessFilePath(ChildProcessId, TEXT("input"), OutPath);
+
+	if (!FFileHelper::SaveStringToFile(JsonTcharText, OutPath.ToString()))
+	{
+		UE_LOG(LogVirtualizationTool, Display, TEXT("Failed to save child process input file '%s'"), OutPath.ToString());
+		return false;
+	}
+
+	return true;
+}
+
+bool FUnrealVirtualizationToolApp::LaunchChildProcess(const FCommand& Command, const FProject& Project, FStringView GlobalOptions, TArray<TUniquePtr<FCommandOutput>>& OutputArray)
+{
+	TRACE_CPUPROFILER_EVENT_SCOPE(LaunchChildProcess);
+
+	UE_LOG(LogVirtualizationTool, Display, TEXT("Launching and waiting on a new instance of the tool..."));
+
+	const FGuid ChildProcessId = FGuid::NewGuid();
+
+	TStringBuilder<512> InputFilePath;
+	if (!TryWriteChildProcessInputFile(ChildProcessId, Command, Project, InputFilePath))
+	{
+		// No need to log an error here, ::TryWriteChildProcessInputFile will take care of that
+		return false;
+	}
+
+	{
+		TRACE_CPUPROFILER_EVENT_SCOPE(RunChildProcess);
+
+		const FString CurrentExePath = FPlatformProcess::ExecutablePath();
+		FString Args = FString::Printf(TEXT("%s -ChildProcess=%s"), *Project.GetProjectFilePath(), InputFilePath.ToString());
+
+		if (!GlobalOptions.IsEmpty())
+		{
+			Args.Append(TEXT(" "));
+			Args.Append(GlobalOptions);
+		}
+
+		const bool bLaunchDetached = false;
+		const bool bLaunchHidden = true;
+		const bool bLaunchReallyHidden = true;
+		uint32 ProcessId = INDEX_NONE;
+		const int32 Priority = 0;
+		const TCHAR* WorkingDirectory = nullptr;
+
+		FProcessPipes Pipes;
+		FProcHandle Handle = FPlatformProcess::CreateProc(*CurrentExePath, *Args, bLaunchDetached, bLaunchHidden, bLaunchReallyHidden,
+			/*OutProcessID*/ nullptr, Priority, WorkingDirectory, Pipes.GetStdIn(), Pipes.GetStdOut());
+
+		{
+			TRACE_CPUPROFILER_EVENT_SCOPE(WaitOnChildProcess);
+			while (FPlatformProcess::IsProcRunning(Handle))
+			{
+				Pipes.ProcessStdOut();
+				FPlatformProcess::Sleep(0.033f);
+			}
+		}
+
+		int32 ReturnCode = INDEX_NONE;
+		if (!FPlatformProcess::GetProcReturnCode(Handle, &ReturnCode))
+		{
+			UE_LOG(LogVirtualizationTool, Display, TEXT("Failed to retrieve the return value of the child process"));
+			return false;
+		}
+
+		if (ReturnCode != 0)
+		{
+			UE_LOG(LogVirtualizationTool, Display, TEXT("Child process failed with error code: %d"), ReturnCode);
+			return false;
+		}
+	}
+
+	if (!TryReadChildProcessOutputFile(ChildProcessId, Command, OutputArray))
+	{
+		// No need to log an error here, ::TryReadChildProcessOutputFile will take care of that
+		return false;
+	}
+
+	return true;
+}
+
 } // namespace UE::Virtualization
diff --git a/Engine/Source/Programs/UnrealVirtualizationTool/Private/UnrealVirtualizationToolApp.h b/Engine/Source/Programs/UnrealVirtualizationTool/Private/UnrealVirtualizationToolApp.h
--- a/Engine/Source/Programs/UnrealVirtualizationTool/Private/UnrealVirtualizationToolApp.h
+++ b/Engine/Source/Programs/UnrealVirtualizationTool/Private/UnrealVirtualizationToolApp.h
@@ -11,6 +11,7 @@ namespace UE::Virtualization
 {
 
 class FCommand;
+struct FCommandOutput;
 
 /** The mode of the application to run */
 enum class EMode :uint32
@@ -58,22 +59,45 @@ private:
 
 	void PrintCmdLineHelp() const;
 
+	bool ProcessProjects(TArray<TUniquePtr<FCommandOutput>>& OutputArray);
+
 	bool TryLoadModules();
 	bool TryInitEnginePlugins();
 	
 	EInitResult TryParseCmdLine();
+	EInitResult TryParseGlobalOptions(const TCHAR* CmdLine);
+	
+	EInitResult CreateCommandFromString(const FString& CommandName, const TCHAR* Cmdline);
 
 	bool TrySortFilesByProject(const TArray<FString>& Packages);
-
 	bool TryFindProject(const FString& PackagePath, FString& ProjectFilePath, FString& PluginFilePath) const;
-
 	FProject& FindOrAddProject(FString&& ProjectFilePath);
 
+	bool IsChildProcess() const;
+
+	void AddGlobalOption(FStringView Options);
+
+	static bool TryReadChildProcessOutputFile(const FGuid& ChildProcessId, const FCommand& Command, TArray<TUniquePtr<FCommandOutput>>& OutputArray);	
+	static bool TryWriteChildProcessOutputFile(const FString& ChildProcessId, const TArray<TUniquePtr<FCommandOutput>>& OutputArray);
+	
+	/** Note not currently static due to dependencies */
+	bool TryReadChildProcessInputFile(const FString& InputPath);
+	static bool TryWriteChildProcessInputFile(const FGuid& ChildProcessId, const FCommand& Command, const FProject& Project, FStringBuilderBase& OutPath);
+
+
+	static bool LaunchChildProcess(const FCommand& Command, const FProject& Project, FStringView GlobalOptions, TArray<TUniquePtr<FCommandOutput>>& OutputArray);
+
 private:
 
 	/** Override used to suppress log messages */
 	TUniquePtr<FFeedbackContext> OutputDeviceOverride;
 
+	/** Stored the id of the child process. This is the base file name of the input file used to launch it */
+	FString ChildProcessId;
+
+	/** A string containing all global options that were found in the cmdline */
+	FString GlobalCmdlineOptions;
+
 	/** Data structure holding the files in the changelist sorted by project and then by plugin*/
 	TArray<FProject> Projects;
 
