[UVT] Improve the output when -help is passed in or the commandline provided is incorrect.

#rnx

- Improve some formatting and make it clear that the virtualize and rehydrate modes should take the project path too.

diff --git a/Engine/Source/Programs/UnrealVirtualizationTool/Private/Commands/RehydrateCommand.cpp b/Engine/Source/Programs/UnrealVirtualizationTool/Private/Commands/RehydrateCommand.cpp
--- a/Engine/Source/Programs/UnrealVirtualizationTool/Private/Commands/RehydrateCommand.cpp
+++ b/Engine/Source/Programs/UnrealVirtualizationTool/Private/Commands/RehydrateCommand.cpp
@@ -22,10 +22,10 @@ FRehydrateCommand::FRehydrateCommand(FStringView CommandName)
 
 void FRehydrateCommand::PrintCmdLineHelp()
 {
+	UE_LOG(LogVirtualizationTool, Display, TEXT("<ProjectFilePath> -Mode=Rehydrate -Package=<string>"));
+	UE_LOG(LogVirtualizationTool, Display, TEXT("<ProjectFilePath> -Mode=Rehydrate -PackageDir=<string>"));
+	UE_LOG(LogVirtualizationTool, Display, TEXT("<ProjectFilePath> -Mode=Rehydrate -Changelist=<number>"));
 	UE_LOG(LogVirtualizationTool, Display, TEXT(""));
-	UE_LOG(LogVirtualizationTool, Display, TEXT("-Mode=Rehydrate -Package=<string>"));
-	UE_LOG(LogVirtualizationTool, Display, TEXT("-Mode=Rehydrate -PackageDir=<string>"));
-	UE_LOG(LogVirtualizationTool, Display, TEXT("-Mode=Rehydrate -Changelist=<number>"));
 }
 
 bool FRehydrateCommand::Initialize(const TCHAR* CmdLine)
diff --git a/Engine/Source/Programs/UnrealVirtualizationTool/Private/Commands/VirtualizeCommand.cpp b/Engine/Source/Programs/UnrealVirtualizationTool/Private/Commands/VirtualizeCommand.cpp
--- a/Engine/Source/Programs/UnrealVirtualizationTool/Private/Commands/VirtualizeCommand.cpp
+++ b/Engine/Source/Programs/UnrealVirtualizationTool/Private/Commands/VirtualizeCommand.cpp
@@ -57,9 +57,9 @@ FVirtualizeCommand::FVirtualizeCommand(FStringView CommandName)
 
 void FVirtualizeCommand::PrintCmdLineHelp()
 {
+	UE_LOG(LogVirtualizationTool, Display, TEXT("<ProjectFilePath> -Mode=Virtualize -Changelist=<number> -Submit [optional]"));
+	UE_LOG(LogVirtualizationTool, Display, TEXT("<ProjectFilePath> -Mode=Virtualize -Path=<string>"));
 	UE_LOG(LogVirtualizationTool, Display, TEXT(""));
-	UE_LOG(LogVirtualizationTool, Display, TEXT("-Mode=Virtualize -Changelist=<number> -Submit [optional]"));
-	UE_LOG(LogVirtualizationTool, Display, TEXT("-Mode=Virtualize -Path=<string>"));
 }
 
 bool FVirtualizeCommand::Initialize(const TCHAR* CmdLine)
diff --git a/Engine/Source/Programs/UnrealVirtualizationTool/Private/UnrealVirtualizationToolApp.cpp b/Engine/Source/Programs/UnrealVirtualizationTool/Private/UnrealVirtualizationToolApp.cpp
--- a/Engine/Source/Programs/UnrealVirtualizationTool/Private/UnrealVirtualizationToolApp.cpp
+++ b/Engine/Source/Programs/UnrealVirtualizationTool/Private/UnrealVirtualizationToolApp.cpp
@@ -303,11 +303,11 @@ void FUnrealVirtualizationToolApp::PrintCmdLineHelp() const
 {
 	UE_LOG(LogVirtualizationTool, Display, TEXT("Usage:"));
 
+	UE_LOG(LogVirtualizationTool, Display, TEXT("Commands:"));
 	// TODO: If the commands were registered in some way we could automate this
 	FVirtualizeCommand::PrintCmdLineHelp(); 
 	FRehydrateCommand::PrintCmdLineHelp();
 
-	UE_LOG(LogVirtualizationTool, Display, TEXT(""));
 	UE_LOG(LogVirtualizationTool, Display, TEXT("Legacy Commands:"));
 	UE_LOG(LogVirtualizationTool, Display, TEXT("-Mode=Changelist -ClientSpecName=<name> [optional] -Changelist=<number> -nosubmit [optional]"));
 	UE_LOG(LogVirtualizationTool, Display, TEXT("-Mode=PackageList -Path=<string>"));
