Fix CIS issue in UnrealVirtualizationTool from CL 24257111

- TCHAR* -> const TCHAR*

diff --git a/Engine/Source/Programs/UnrealVirtualizationTool/Private/UnrealVirtualizationToolApp.cpp b/Engine/Source/Programs/UnrealVirtualizationTool/Private/UnrealVirtualizationToolApp.cpp
--- a/Engine/Source/Programs/UnrealVirtualizationTool/Private/UnrealVirtualizationToolApp.cpp
+++ b/Engine/Source/Programs/UnrealVirtualizationTool/Private/UnrealVirtualizationToolApp.cpp
@@ -55,7 +55,7 @@ void LexFromString(UE::Virtualization::EMode& OutValue, const FStringView& InStr
  * Utility to convert EMode to a string. NOTE that legacy modes will return the
  * name of the newer mode that replaces it and not the legacy mode.
  */
-TCHAR* LexToString(UE::Virtualization::EMode Mode)
+const TCHAR* LexToString(UE::Virtualization::EMode Mode)
 {
 	using namespace UE::Virtualization;
 
