[UVT] Replace the cmdline parameter -verbose with -minimallogging and invert the meaning.

- This means by default all log messages with 'Display' verbosity will be shown, we will only demote them to 'Log' verbosity when -minimallogging is passed in.
- UVT was originally intended to be called as a custom tool in P4V, at which point it made sense to hide a lot of the logging from the user so that the output would be easier for them to follow and if they did encounter a problem they could send the full log file in the bug report.
- In the past few months there has been an increase in people trying to use UVT to verify that they have setup VA on their project correctly rather than using the editor. In many of these cases it was not obvious that they'd need to pass in -verbose to see full output.
- By inverting the cmdline meaning we can display all logging to people by default (i.e. calling from the command line) and then supress the logging when called via P4V by using the new cmdline.

diff --git a/Engine/Source/Programs/UnrealVirtualizationTool/Private/UnrealVirtualizationToolApp.cpp b/Engine/Source/Programs/UnrealVirtualizationTool/Private/UnrealVirtualizationToolApp.cpp
--- a/Engine/Source/Programs/UnrealVirtualizationTool/Private/UnrealVirtualizationToolApp.cpp
+++ b/Engine/Source/Programs/UnrealVirtualizationTool/Private/UnrealVirtualizationToolApp.cpp
@@ -314,7 +314,7 @@ void FUnrealVirtualizationToolApp::PrintCmdLineHelp() const
 	
 	UE_LOG(LogVirtualizationTool, Display, TEXT(""));
 	UE_LOG(LogVirtualizationTool, Display, TEXT("Global Options:"));
-	UE_LOG(LogVirtualizationTool, Display, TEXT("\t-verbose (all log messages with display verbosity will be displayed, not just LogVirtualizationTool)"));
+	UE_LOG(LogVirtualizationTool, Display, TEXT("\t-MinimalLogging (demote log messages with 'display' verbosity to 'log' verbosity except those using the LogVirtualizationTool category)"));
 }
 
 bool FUnrealVirtualizationToolApp::TryLoadModules()
@@ -414,16 +414,16 @@ EInitResult FUnrealVirtualizationToolApp::TryParseGlobalOptions(const TCHAR* Cmd
 {
 	GlobalCmdlineOptions.Reset();
 
-	if (FParse::Param(CmdLine, TEXT("verbose")))
+	if (FParse::Param(CmdLine, TEXT("MinimalLogging")))
 	{
-		UE_LOG(LogVirtualizationTool, Display, TEXT("Cmdline parameter '-verbose' found, no longer supressing Display log messages!"));
-		OutputDeviceOverride.Reset();
+		UE_LOG(LogVirtualizationTool, Display, TEXT("Cmdline parameter '-MinimalLogging' found, supressing Display log messages"));
+		OutputDeviceOverride = MakeUnique<FOverrideOutputDevice>();
 
-		AddGlobalOption(TEXT("-verbose"));	
+		AddGlobalOption(TEXT("-MinimalLogging"));	
 	}
 	else
 	{
-		OutputDeviceOverride = MakeUnique<FOverrideOutputDevice>();
+		OutputDeviceOverride.Reset();
 	}
 
 	return EInitResult::Success;
