Add the json module to UnrealVirtualizationTool.

#rnx

- Allows UVT to compile in modular mode.

diff --git a/Engine/Source/Programs/UnrealVirtualizationTool/UnrealVirtualizationTool.Build.cs b/Engine/Source/Programs/UnrealVirtualizationTool/UnrealVirtualizationTool.Build.cs
--- a/Engine/Source/Programs/UnrealVirtualizationTool/UnrealVirtualizationTool.Build.cs
+++ b/Engine/Source/Programs/UnrealVirtualizationTool/UnrealVirtualizationTool.Build.cs
@@ -16,7 +16,8 @@ public class UnrealVirtualizationTool : ModuleRules
 				"ApplicationCore",
 				"SourceControl",
 				"PerforceSourceControl",
-				"DesktopPlatform"
+				"DesktopPlatform",
+				"Json"
 			});
 
 		PublicIncludePathModuleNames.Add("Launch");
