[UVT] Remove the bespoke checkout code for the rehydration command as the rehydration process takes care of that now.

#rnx

diff --git a/Engine/Source/Programs/UnrealVirtualizationTool/Private/Commands/CommandBase.cpp b/Engine/Source/Programs/UnrealVirtualizationTool/Private/Commands/CommandBase.cpp
--- a/Engine/Source/Programs/UnrealVirtualizationTool/Private/Commands/CommandBase.cpp
+++ b/Engine/Source/Programs/UnrealVirtualizationTool/Private/Commands/CommandBase.cpp
@@ -144,60 +144,6 @@ bool FCommand::TryConnectToSourceControl(FStringView ClientSpecName)
 	}
 }
 
-bool FCommand::TryCheckOutFilesForProject(FStringView ClientSpecName, FStringView ProjectRoot, const TArray<FString>& ProjectPackages)
-{
-	// Override the root directory for source control to use the project for which we are trying to hydrate packages for
-	ISourceControlModule::Get().RegisterSourceControlProjectDirDelegate(FSourceControlProjectDirDelegate::CreateLambda([&ProjectRoot]()
-		{
-			return *WriteToString<260>(ProjectRoot, TEXT("/"));
-		}));
-
-	ON_SCOPE_EXIT
-	{
-		ISourceControlModule::Get().UnregisterSourceControlProjectDirDelegate();
-	};
-
-	if (!TryConnectToSourceControl(ClientSpecName))
-	{
-		return false;
-	}
-
-	UE_LOG(LogVirtualizationTool, Display, TEXT("\t\tChecking status of package files..."));
-
-	TArray<TSharedRef<ISourceControlState>> FileStates;
-	if (SCCProvider->GetState(ProjectPackages, FileStates, EStateCacheUsage::ForceUpdate) != ECommandResult::Succeeded)
-	{
-		UE_LOG(LogVirtualizationTool, Error, TEXT("Failed to find file states for packages from source control"));
-		return false;
-	}
-
-	TArray<FString> FilesToCheckout;
-	FilesToCheckout.Reserve(FileStates.Num());
-
-	for (const TSharedRef<ISourceControlState>& State : FileStates)
-	{
-		if (State->CanCheckout())
-		{
-			FilesToCheckout.Add(State->GetFilename());
-		}
-	}
-
-	if (!FilesToCheckout.IsEmpty())
-	{
-		UE_LOG(LogVirtualizationTool, Display, TEXT("\t\tChecking out %d file(s) from source control..."), FilesToCheckout.Num());
-
-		if (SCCProvider->Execute(ISourceControlOperation::Create<FCheckOut>(), FilesToCheckout, EConcurrency::Synchronous) != ECommandResult::Succeeded)
-		{
-			UE_LOG(LogVirtualizationTool, Error, TEXT("Failed to checkout packages from source control"));
-			return false;
-		}
-	}
-
-	UE_LOG(LogVirtualizationTool, Display, TEXT("\t\tAll files checked out and writable"));
-
-	return true;
-}
-
 bool FCommand::TryParseChangelist(FStringView ClientSpecName, FStringView ChangelistNumber, TArray<FString>& OutPackages, FSourceControlChangelistPtr* OutChangelist)
 {
 	TRACE_CPUPROFILER_EVENT_SCOPE(TryParseChangelist);
diff --git a/Engine/Source/Programs/UnrealVirtualizationTool/Private/Commands/CommandBase.h b/Engine/Source/Programs/UnrealVirtualizationTool/Private/Commands/CommandBase.h
--- a/Engine/Source/Programs/UnrealVirtualizationTool/Private/Commands/CommandBase.h
+++ b/Engine/Source/Programs/UnrealVirtualizationTool/Private/Commands/CommandBase.h
@@ -134,8 +134,6 @@ protected: // Common SourceControl Code
 		return TryConnectToSourceControl(FStringView());
 	}
 
-	bool TryCheckOutFilesForProject(FStringView ClientSpecName, FStringView ProjectRoot, const TArray<FString>& ProjectPackages);
-	
 	bool TryParseChangelist(FStringView ClientSpecName, FStringView ChangelistNumber, TArray<FString>& OutPackages, FSourceControlChangelistPtr* OutChangelist);
 
 	FString FindClientSpecForChangelist(FStringView ChangelistNumber);
diff --git a/Engine/Source/Programs/UnrealVirtualizationTool/Private/Commands/RehydrateCommand.cpp b/Engine/Source/Programs/UnrealVirtualizationTool/Private/Commands/RehydrateCommand.cpp
--- a/Engine/Source/Programs/UnrealVirtualizationTool/Private/Commands/RehydrateCommand.cpp
+++ b/Engine/Source/Programs/UnrealVirtualizationTool/Private/Commands/RehydrateCommand.cpp
@@ -104,15 +104,9 @@ bool FRehydrateCommand::ProcessProject(const FProject& Project, TUniquePtr<FComm
 		UE::Virtualization::Shutdown();
 	};
 
-	const TArray<FString> ProjectPackages = Project.GetAllPackages();
-
-	if (!TryCheckOutFilesForProject(ClientSpecName, Project.GetProjectRoot(), ProjectPackages))
-	{
-		return false;
-	}
-
 	UE_LOG(LogVirtualizationTool, Display, TEXT("\t\tAttempting to rehydrate packages..."), ProjectName.ToString());
 
+	const TArray<FString> ProjectPackages = Project.GetAllPackages();
 
 	ERehydrationOptions Options = ERehydrationOptions::None;
 	if (bShouldCheckout)
