[UVT] Properly apply the project path being supplied via the commandline (if there is one)

#rnx

- Although we now relaunch the tool with the correct project path (if not supplied) this was not enough to allow us to load the projects config files and so set up the DDC and source control etc properly. My original testing was picking up hacks in the programs config file that was allowing it to work.
- FGenericPlatformMisc::ProjectDir has a bespoke path for standalone programs which means even if the uproject is on the cmdline it is ignored.
- We probably should look into a way for a project to mark itself as being "project aware" but for now we can make use of the project path override system to make sure it is set correctly, which will fix the problem for now.

diff --git a/Engine/Source/Programs/UnrealVirtualizationTool/Private/UnrealVirtualizationTool.cpp b/Engine/Source/Programs/UnrealVirtualizationTool/Private/UnrealVirtualizationTool.cpp
--- a/Engine/Source/Programs/UnrealVirtualizationTool/Private/UnrealVirtualizationTool.cpp
+++ b/Engine/Source/Programs/UnrealVirtualizationTool/Private/UnrealVirtualizationTool.cpp
@@ -14,6 +14,33 @@ bool UnrealVirtualizationToolMain(int32 ArgC, TCHAR* ArgV[])
 {
 	using namespace UE::Virtualization;
 
+	// Although standalone tools can set the project path via the cmdline this will not change the ProjectDir being
+	// used as standalone tools have a beespoke path in FGenericPlatformMisc::ProjectDir. We can work around this
+	// for now by doing our own parsing then using the project dir override feature.
+	// This is a band aid while we consider adding better support for project files/directories with stand alone tools.
+	if(ArgC >= 2)
+	{
+		FString Cmd(ArgV[1]);
+
+		if (!Cmd.IsEmpty() && !Cmd.StartsWith(TEXT("-")) && Cmd.EndsWith(FProjectDescriptor::GetExtension()))
+		{
+			FString ProjectDir = FPaths::GetPath(Cmd);
+			ProjectDir = FFileManagerGeneric::DefaultConvertToRelativePath(*ProjectDir);
+			
+			// The path should end with a trailing slash (see FGenericPlatformMisc::ProjectDir) so we use
+			// NormalizeFilename not NormalizeDirectoryName as the latter will remove trailing slashes. We
+			// also need to add one if it is missing.
+			// We probably should move this path fixup code to 'FPlatformMisc::SetOverrideProjectDir'
+			FPaths::NormalizeFilename(ProjectDir);
+			if (!ProjectDir.EndsWith(TEXT("/")))
+			{
+				ProjectDir += TEXT("/");
+			}
+
+			FPlatformMisc::SetOverrideProjectDir(ProjectDir);
+		}
+	}
+
 	GEngineLoop.PreInit(ArgC, ArgV);
 	check(GConfig && GConfig->IsReadyForUse());
 
