[UVT] Fix issue preventing the rehydration command from checking out packages on demand

#rnx

- The checkout cmdline option was not being serialized to the child process correctly.

diff --git a/Engine/Source/Programs/UnrealVirtualizationTool/Private/Commands/RehydrateCommand.cpp b/Engine/Source/Programs/UnrealVirtualizationTool/Private/Commands/RehydrateCommand.cpp
--- a/Engine/Source/Programs/UnrealVirtualizationTool/Private/Commands/RehydrateCommand.cpp
+++ b/Engine/Source/Programs/UnrealVirtualizationTool/Private/Commands/RehydrateCommand.cpp
@@ -76,7 +76,7 @@ void FRehydrateCommand::Serialize(FJsonSerializerBase& Serializer)
 {
 	TRACE_CPUPROFILER_EVENT_SCOPE(FRehydrateCommand::Serialize);
 
-	// The command has no additional data when launched as a child process
+	Serializer.Serialize(TEXT("ShouldCheckout"), bShouldCheckout);
 }
 
 bool FRehydrateCommand::ProcessProject(const FProject& Project, TUniquePtr<FCommandOutput>& Output)
