[UVT] The tool now supports the '-fastexit' commandline option

#rnx

- The main goal was to prevent the GConfig from flushing to disk on shutdown as we do not want the tool to do that.
- We can do this by disabling the disk operations for GConfig or forcing the process to terminate, the latter of which we support in commandlets via the cmdline option '-fastexit'

diff --git a/Engine/Source/Programs/UnrealVirtualizationTool/Private/UnrealVirtualizationTool.cpp b/Engine/Source/Programs/UnrealVirtualizationTool/Private/UnrealVirtualizationTool.cpp
--- a/Engine/Source/Programs/UnrealVirtualizationTool/Private/UnrealVirtualizationTool.cpp
+++ b/Engine/Source/Programs/UnrealVirtualizationTool/Private/UnrealVirtualizationTool.cpp
@@ -10,8 +10,10 @@ IMPLEMENT_APPLICATION(UnrealVirtualizationTool, "UnrealVirtualizationTool");
 
 DEFINE_LOG_CATEGORY(LogVirtualizationTool);
 
-bool UnrealVirtualizationToolMain(int32 ArgC, TCHAR* ArgV[])
+int32 UnrealVirtualizationToolMain(int32 ArgC, TCHAR* ArgV[])
 {
+	TRACE_CPUPROFILER_EVENT_SCOPE(UnrealVirtualizationToolMain);
+
 	using namespace UE::Virtualization;
 
 	// Although standalone tools can set the project path via the cmdline this will not change the ProjectDir being
@@ -76,18 +78,31 @@ bool UnrealVirtualizationToolMain(int32 ArgC, TCHAR* ArgV[])
 
 	UE_CLOG(bRanSuccessfully, LogVirtualizationTool, Display, TEXT("UnrealVirtualizationTool ran successfully"));
 
-	// Even though we are exiting anyway we need to request an engine exit in order to get a clean shutdown
-	RequestEngineExit(TEXT("The process has finished"));
+	const uint8 ReturnCode = bRanSuccessfully ? 0 : 1;
+
+	if (FParse::Param(FCommandLine::Get(), TEXT("fastexit")))
+	{
+		FPlatformMisc::RequestExitWithStatus(true, ReturnCode);
+	}
+	else
+	{
+		TRACE_CPUPROFILER_EVENT_SCOPE(Shutdown);
+
+		GConfig->DisableFileOperations(); // We don't want to write out any config file changes!
 
-	FEngineLoop::AppPreExit();
-	FModuleManager::Get().UnloadModulesAtShutdown();
-	FEngineLoop::AppExit();
+		// Even though we are exiting anyway we need to request an engine exit in order to get a clean shutdown
+		RequestEngineExit(TEXT("The process has finished"));
+
+		FEngineLoop::AppPreExit();
+		FModuleManager::Get().UnloadModulesAtShutdown();
+		FEngineLoop::AppExit();
+	}
 
-	return bRanSuccessfully;
+	return ReturnCode;
 }
 
 INT32_MAIN_INT32_ARGC_TCHAR_ARGV()
 {
 	FTaskTagScope Scope(ETaskTag::EGameThread);
-	return UnrealVirtualizationToolMain(ArgC, ArgV) ? 0 : 1;
+	return UnrealVirtualizationToolMain(ArgC, ArgV);
 }
\ No newline at end of file
