[UVT] Improve the reliability of establishing revision control connections when using the tool to checkout or submit packages

#rnx

- Changed the logic of TryConnectToSourceControl so that if a project is set we try to use the projects existing revision control settings for the default provider (which is the provider that will be used by the virtualization/rehydration process), if the project is not set then the only revision control operation that should occur is parsing of an existing changelist to find the files that we might want to operate on.
-- As we now also use the default provider we cannot use a single TUniquePtr as we do not own the pointer to the default provider. So there is a general pointer to be used for the current active provider and a TUniquePtr if we own the provider.
- Added calls to ::TryConnectToSourceControl if a command is setting the Checkout flag which should ensure that the packages will be able to be checked out.

diff --git a/Engine/Source/Programs/UnrealVirtualizationTool/Private/Commands/CommandBase.cpp b/Engine/Source/Programs/UnrealVirtualizationTool/Private/Commands/CommandBase.cpp
--- a/Engine/Source/Programs/UnrealVirtualizationTool/Private/Commands/CommandBase.cpp
+++ b/Engine/Source/Programs/UnrealVirtualizationTool/Private/Commands/CommandBase.cpp
@@ -110,7 +110,7 @@ FCommand::EPathResult FCommand::ParseSwitchForPaths(const FString& Switch, TArra
 
 bool FCommand::TryConnectToSourceControl(FStringView ClientSpecName)
 {
-	if (SCCProvider.IsValid())
+	if (SCCProvider != nullptr)
 	{
 		// Already connected so just return
 		return true;
@@ -118,29 +118,51 @@ bool FCommand::TryConnectToSourceControl(FStringView ClientSpecName)
 
 	UE_LOG(LogVirtualizationTool, Log, TEXT("Trying to connect to source control..."));
 
-	FSourceControlInitSettings SCCSettings(FSourceControlInitSettings::EBehavior::OverrideAll);
-	SCCSettings.AddSetting(TEXT("P4Client"), ClientSpecName);
-
-	SCCProvider = ISourceControlModule::Get().CreateProvider(FName("Perforce"), TEXT("UnrealVirtualizationTool"), SCCSettings);
-	if (SCCProvider.IsValid())
+	if (FPaths::IsProjectFilePathSet())
 	{
+		// If the project has been set then we can use the default source control provider
+		if (!ISourceControlModule::Get().GetProvider().IsEnabled())
+		{
+			// TODO - Warning
+			ISourceControlModule::Get().SetProvider(FName("Perforce"));
+		}
+
+		SCCProvider = &ISourceControlModule::Get().GetProvider();
 		SCCProvider->Init(true);
 
-		if (SCCProvider->IsAvailable())
+		return SCCProvider->IsAvailable();
+	}
+	else
+	{
+		// If we do not have a project set we need to create our own provider. We make the assumption that
+		// we should use perforce since the only reason we should need a provider without a project is if
+		// we are parsing a perforce changelist for files to operate on.
+		FSourceControlInitSettings SCCSettings(FSourceControlInitSettings::EBehavior::OverrideAll);
+		SCCSettings.SetConfigBehavior(FSourceControlInitSettings::EConfigBehavior::ReadOnly);
+		SCCSettings.AddSetting(TEXT("P4Client"), ClientSpecName);
+
+		OwnedSCCProvider = ISourceControlModule::Get().CreateProvider(FName("Perforce"), TEXT("UnrealVirtualizationTool"), SCCSettings);
+		if (OwnedSCCProvider.IsValid())
 		{
-			return true;
+			SCCProvider = OwnedSCCProvider.Get();
+			SCCProvider->Init(true);
+
+			if (SCCProvider->IsAvailable())
+			{
+				return true;
+			}
+			else
+			{
+				UE_LOG(LogVirtualizationTool, Error, TEXT("Failed to establish a perforce connection"));
+				return false;
+			}
+
 		}
 		else
 		{
-			UE_LOG(LogVirtualizationTool, Error, TEXT("Failed to establish a perforce connection"));
+			UE_LOG(LogVirtualizationTool, Error, TEXT("Failed to instantiate a perforce revision control connection"));
 			return false;
 		}
-		
-	}
-	else
-	{
-		UE_LOG(LogVirtualizationTool, Error, TEXT("Failed to instantiate a perforce revision control connection"));
-		return false;
 	}
 }
 
@@ -157,7 +179,7 @@ bool FCommand::TryParseChangelist(FStringView ClientSpecName, FStringView Change
 		return false;
 	}
 
-	if (!SCCProvider.IsValid())
+	if (SCCProvider == nullptr)
 	{
 		UE_LOG(LogVirtualizationTool, Error, TEXT("No valid source control connection found!"));
 		return false;
@@ -244,7 +266,7 @@ FString FCommand::FindClientSpecForChangelist(FStringView ChangelistNumber)
 		return FString();
 	}
 
-	if (!SCCProvider.IsValid())
+	if (SCCProvider == nullptr)
 	{
 		UE_LOG(LogVirtualizationTool, Error, TEXT("No valid source control connection found!"));
 		return FString();
diff --git a/Engine/Source/Programs/UnrealVirtualizationTool/Private/Commands/CommandBase.h b/Engine/Source/Programs/UnrealVirtualizationTool/Private/Commands/CommandBase.h
--- a/Engine/Source/Programs/UnrealVirtualizationTool/Private/Commands/CommandBase.h
+++ b/Engine/Source/Programs/UnrealVirtualizationTool/Private/Commands/CommandBase.h
@@ -141,8 +141,13 @@ protected: // Common SourceControl Code
 private:
 	FString CommandName;
 
+	TUniquePtr<ISourceControlProvider> OwnedSCCProvider;
+
 protected:
-	TUniquePtr<ISourceControlProvider> SCCProvider;
+
+	ISourceControlProvider* SCCProvider = nullptr;
+
+	
 };
 
 } // namespace UE::Virtualization
\ No newline at end of file
diff --git a/Engine/Source/Programs/UnrealVirtualizationTool/Private/Commands/RehydrateCommand.cpp b/Engine/Source/Programs/UnrealVirtualizationTool/Private/Commands/RehydrateCommand.cpp
--- a/Engine/Source/Programs/UnrealVirtualizationTool/Private/Commands/RehydrateCommand.cpp
+++ b/Engine/Source/Programs/UnrealVirtualizationTool/Private/Commands/RehydrateCommand.cpp
@@ -112,6 +112,9 @@ bool FRehydrateCommand::ProcessProject(const FProject& Project, TUniquePtr<FComm
 	if (bShouldCheckout)
 	{
 		Options |= ERehydrationOptions::Checkout;
+
+		// Make sure that we have a valid source control connection if we might try to checkout packages
+		TryConnectToSourceControl();
 	}
 
 	FRehydrationResult Result = UE::Virtualization::IVirtualizationSystem::Get().TryRehydratePackages(ProjectPackages, Options);
diff --git a/Engine/Source/Programs/UnrealVirtualizationTool/Private/Commands/VirtualizeCommand.cpp b/Engine/Source/Programs/UnrealVirtualizationTool/Private/Commands/VirtualizeCommand.cpp
--- a/Engine/Source/Programs/UnrealVirtualizationTool/Private/Commands/VirtualizeCommand.cpp
+++ b/Engine/Source/Programs/UnrealVirtualizationTool/Private/Commands/VirtualizeCommand.cpp
@@ -180,6 +180,9 @@ bool FVirtualizeCommand::ProcessProject(const FProject& Project, TUniquePtr<FCom
 	if (bShouldCheckout)
 	{
 		Options |= EVirtualizationOptions::Checkout;
+
+		// Make sure that we have a valid source control connection if we might try to checkout packages
+		TryConnectToSourceControl();
 	}
 
 	FVirtualizationResult Result = UE::Virtualization::IVirtualizationSystem::Get().TryVirtualizePackages(ProjectPackages, Options);
