[UVT] Allow the tool to discover custom plugins

#rnx

- When the tool was created we tried to avoid discovering all plugins as it was very slow. Since then it has been optimized and only costs 500ms or so, which for now is acceptable. So now we allow plugins to be discovered in the normal way. If the plugins have been set up correctly to work with the tool then they will be loaded.
- Note that the perforce plugin still needs to be mounted manually, but no longer needs the paths adding.
- Removed UE_DISABLE_PLUGIN_DISCOVERY from the Target.cs so that we scan for all plugins on startup.
- Removed adding of perforce to AdditionalPlugins in the Target.cs as it wasn't really doing anything.

diff --git a/Engine/Source/Programs/UnrealVirtualizationTool/Private/UnrealVirtualizationToolApp.cpp b/Engine/Source/Programs/UnrealVirtualizationTool/Private/UnrealVirtualizationToolApp.cpp
--- a/Engine/Source/Programs/UnrealVirtualizationTool/Private/UnrealVirtualizationToolApp.cpp
+++ b/Engine/Source/Programs/UnrealVirtualizationTool/Private/UnrealVirtualizationToolApp.cpp
@@ -333,25 +333,27 @@ bool FUnrealVirtualizationToolApp::TryInitEnginePlugins()
 
 	UE_LOG(LogVirtualizationTool, Log, TEXT("Loading Engine Plugins"));
 
-	IPluginManager& PluginMgr = IPluginManager::Get();
+	auto LoadPlugin = [](const FString& PlugInName) -> bool
+		{
+			IPluginManager& PluginMgr = IPluginManager::Get();
 
-	const FString PerforcePluginPath = FPaths::EnginePluginsDir() / TEXT("Developer/PerforceSourceControl/PerforceSourceControl.uplugin");
-	FText ErrorMsg;
-	if (!PluginMgr.AddToPluginsList(PerforcePluginPath, &ErrorMsg))
-	{
-		UE_LOG(LogVirtualizationTool, Error, TEXT("Failed to find 'PerforceSourceControl' plugin due to: %s"), *ErrorMsg.ToString());
-		return false;
-	}
+			PluginMgr.MountNewlyCreatedPlugin(PlugInName);
 
-	PluginMgr.MountNewlyCreatedPlugin(TEXT("PerforceSourceControl"));
+			TSharedPtr<IPlugin> Plugin = PluginMgr.FindPlugin(PlugInName);
+			if (Plugin == nullptr || !Plugin->IsEnabled())
+			{
+				UE_LOG(LogVirtualizationTool, Error, TEXT("The plugin '%s' is disabled."), *PlugInName);
+				return false;
+			}
+
+			return true;
+		};
 
-	TSharedPtr<IPlugin> Plugin = PluginMgr.FindPlugin(TEXT("PerforceSourceControl"));
-	if (Plugin == nullptr || !Plugin->IsEnabled())
+	if (!LoadPlugin(TEXT("PerforceSourceControl")))
 	{
-		UE_LOG(LogVirtualizationTool, Error, TEXT("The 'PerforceSourceControl' plugin is disabled."));
 		return false;
 	}
-
+	
 	return true;
 }
 
diff --git a/Engine/Source/Programs/UnrealVirtualizationTool/UnrealVirtualizationTool.Target.cs b/Engine/Source/Programs/UnrealVirtualizationTool/UnrealVirtualizationTool.Target.cs
--- a/Engine/Source/Programs/UnrealVirtualizationTool/UnrealVirtualizationTool.Target.cs
+++ b/Engine/Source/Programs/UnrealVirtualizationTool/UnrealVirtualizationTool.Target.cs
@@ -30,10 +30,7 @@ public class UnrealVirtualizationToolTarget : TargetRules
 		// Enable Developer plugins
 		bCompileWithPluginSupport = true;
 
-		AdditionalPlugins.Add("PerforceSourceControl");
-
 		GlobalDefinitions.Add("UE_TRACE_ENABLED=1");
 		GlobalDefinitions.Add("UE_SUPPORT_FULL_PACKAGEPATH=1");
-		GlobalDefinitions.Add("UE_DISABLE_PLUGIN_DISCOVERY=1");
 	}
 }
