Log the time taken to virtualize/hydrate when using the UnrealVirtualizationTool

#rnx

- Virtualization/hydration can be quite slow. We should log the time taken so that it is easier for users to see if something is abnormally slow and report it.
- Small code cleanup in the hydration code to move the ::UnRegisterMountPoints call to a ON_SCOPE_EXIT

diff --git a/Engine/Source/Programs/UnrealVirtualizationTool/Private/Commands/RehydrateCommand.cpp b/Engine/Source/Programs/UnrealVirtualizationTool/Private/Commands/RehydrateCommand.cpp
--- a/Engine/Source/Programs/UnrealVirtualizationTool/Private/Commands/RehydrateCommand.cpp
+++ b/Engine/Source/Programs/UnrealVirtualizationTool/Private/Commands/RehydrateCommand.cpp
@@ -96,6 +96,11 @@ bool FRehydrateCommand::ProcessProject(const FProject& Project, TUniquePtr<FComm
 
 	Project.RegisterMountPoints();
 
+	ON_SCOPE_EXIT
+	{
+		Project.UnRegisterMountPoints();
+	};
+
 	UE::Virtualization::FInitParams InitParams(ProjectName, EngineConfigWithProject);
 	UE::Virtualization::Initialize(InitParams, UE::Virtualization::EInitializationFlags::ForceInitialize);
 
@@ -133,8 +138,7 @@ bool FRehydrateCommand::ProcessProject(const FProject& Project, TUniquePtr<FComm
 		UE_LOG(LogVirtualizationTool, Display, TEXT("\t\t%d packages were checked out of revision control"), Result.CheckedOutPackages.Num());
 	}
 
-	Project.UnRegisterMountPoints();
-
+	UE_LOG(LogVirtualizationTool, Display, TEXT("\t\tTime taken %.2f(s)"), Result.TimeTaken);
 	UE_LOG(LogVirtualizationTool, Display, TEXT("\t\tRehyration of project packages complete!"), ProjectName.ToString());
 
 	return true;
diff --git a/Engine/Source/Programs/UnrealVirtualizationTool/Private/Commands/VirtualizeCommand.cpp b/Engine/Source/Programs/UnrealVirtualizationTool/Private/Commands/VirtualizeCommand.cpp
--- a/Engine/Source/Programs/UnrealVirtualizationTool/Private/Commands/VirtualizeCommand.cpp
+++ b/Engine/Source/Programs/UnrealVirtualizationTool/Private/Commands/VirtualizeCommand.cpp
@@ -202,6 +202,7 @@ bool FVirtualizeCommand::ProcessProject(const FProject& Project, TUniquePtr<FCom
 		UE_LOG(LogVirtualizationTool, Display, TEXT("\t\t%d packages were checked out of revision control"), Result.CheckedOutPackages.Num());
 	}
 
+	UE_LOG(LogVirtualizationTool, Display, TEXT("\t\tTime taken %.2f(s)"), Result.TimeTaken);
 	UE_LOG(LogVirtualizationTool, Display, TEXT("\t\tVirtualization of project packages complete!"), ProjectName.ToString());
 
 	Output = MakeUnique<FVirtualizeCommandOutput>(Project.GetProjectName(), Result.DescriptionTags);
