No master fix for UnrealVirtualizationTool

diff --git a/Engine/Source/Developer/DerivedDataCache/Private/DerivedDataCacheStoreThrottle.cpp b/Engine/Source/Developer/DerivedDataCache/Private/DerivedDataCacheStoreThrottle.cpp
--- a/Engine/Source/Developer/DerivedDataCache/Private/DerivedDataCacheStoreThrottle.cpp
+++ b/Engine/Source/Developer/DerivedDataCache/Private/DerivedDataCacheStoreThrottle.cpp
@@ -7,6 +7,7 @@
 #include "DerivedDataLegacyCacheStore.h"
 #include "HAL/PlatformProcess.h"
 #include "HAL/PlatformTime.h"
+#include "ProfilingDebugging/CpuProfilerTrace.h"
 #include <atomic>
 
 namespace UE::DerivedData
diff --git a/Engine/Source/Programs/UnrealVirtualizationTool/Private/UnrealVirtualizationToolApp.cpp b/Engine/Source/Programs/UnrealVirtualizationTool/Private/UnrealVirtualizationToolApp.cpp
--- a/Engine/Source/Programs/UnrealVirtualizationTool/Private/UnrealVirtualizationToolApp.cpp
+++ b/Engine/Source/Programs/UnrealVirtualizationTool/Private/UnrealVirtualizationToolApp.cpp
@@ -132,7 +132,9 @@ public:
 
 	virtual void Serialize(const TCHAR* V, ELogVerbosity::Type Verbosity, const FName& Category, double Time) override
 	{
+#if !NO_LOGGING
 		if (Verbosity == ELogVerbosity::Display && Category != LogVirtualizationTool.GetCategoryName())
+#endif
 		{
 			Verbosity = ELogVerbosity::Log;
 		}
@@ -142,12 +144,14 @@ public:
 
 	virtual void SerializeRecord(const UE::FLogRecord& Record) override
 	{
+#if !NO_LOGGING
 		if (Record.GetVerbosity() == ELogVerbosity::Display && Record.GetCategory() != LogVirtualizationTool.GetCategoryName())
 		{
 			UE::FLogRecord LocalRecord = Record;
 			LocalRecord.SetVerbosity(ELogVerbosity::Log);
 			return FFeedbackContextAnsi::SerializeRecord(LocalRecord);
 		}
+#endif
 
 		FFeedbackContextAnsi::SerializeRecord(Record);
 	}
diff --git a/Engine/Source/Runtime/TraceLog/Private/Trace/Control.cpp b/Engine/Source/Runtime/TraceLog/Private/Trace/Control.cpp
--- a/Engine/Source/Runtime/TraceLog/Private/Trace/Control.cpp
+++ b/Engine/Source/Runtime/TraceLog/Private/Trace/Control.cpp
@@ -6,6 +6,7 @@
 
 #include "Trace/Platform.h"
 #include "Trace/Message.h"
+#include "Trace/Detail/Channel.h"
 
 #include "Misc/CString.h"
 #include "Templates/UnrealTemplate.h"
