- Fixed a deprecation from config system changes

diff --git a/Engine/Source/Programs/UnrealVirtualizationTool/Private/ProjectFiles.cpp b/Engine/Source/Programs/UnrealVirtualizationTool/Private/ProjectFiles.cpp
--- a/Engine/Source/Programs/UnrealVirtualizationTool/Private/ProjectFiles.cpp
+++ b/Engine/Source/Programs/UnrealVirtualizationTool/Private/ProjectFiles.cpp
@@ -36,7 +36,7 @@ void ConvertToMountPoint(const FString& ProjectFilePath, FString& OutRootPath, F
 void ReloadConfigs(FConfigFile& PluginConfig)
 {
 	// Reload configs so objects get the changes
-	for (const auto& ConfigEntry : PluginConfig)
+	for (const auto& ConfigEntry : (const FConfigFile&)PluginConfig)
 	{
 		// Skip out if someone put a config section in the INI without any actual data
 		if (ConfigEntry.Value.Num() == 0)
