[UVT] Minor logging improvements

#rnx

- Do not log if '-MinimalLogging' is set in the cmdline. If it is set then the tool is supposed to only display non-technical info and the user probably doesn't care about the flag. If we need to check the state of it, the log file will contain the processes cmdline anyway.
- Add a missing tab to the logging in FUnrealVirtualizationToolApp::TrySortFilesByProject to keep it aligned.

diff --git a/Engine/Source/Programs/UnrealVirtualizationTool/Private/UnrealVirtualizationToolApp.cpp b/Engine/Source/Programs/UnrealVirtualizationTool/Private/UnrealVirtualizationToolApp.cpp
--- a/Engine/Source/Programs/UnrealVirtualizationTool/Private/UnrealVirtualizationToolApp.cpp
+++ b/Engine/Source/Programs/UnrealVirtualizationTool/Private/UnrealVirtualizationToolApp.cpp
@@ -422,7 +422,6 @@ EInitResult FUnrealVirtualizationToolApp::TryParseGlobalOptions(const TCHAR* Cmd
 
 	if (FParse::Param(CmdLine, TEXT("MinimalLogging")))
 	{
-		UE_LOG(LogVirtualizationTool, Display, TEXT("Cmdline parameter '-MinimalLogging' found, supressing Display log messages"));
 		OutputDeviceOverride = MakeUnique<FOverrideOutputDevice>();
 
 		AddGlobalOption(TEXT("-MinimalLogging"));	
@@ -473,7 +472,7 @@ bool FUnrealVirtualizationToolApp::TrySortFilesByProject(const TArray<FString>&
 {
 	TRACE_CPUPROFILER_EVENT_SCOPE(TrySortFilesByProject);
 
-	UE_LOG(LogVirtualizationTool, Display, TEXT("Sorting files by project..."));
+	UE_LOG(LogVirtualizationTool, Display, TEXT("\tSorting files by project..."));
 
 	for (const FString& PackagePath : Packages)
 	{
