[UVT] Fix a bug where running the tool on a changelist from a project that has not enabled virtualization would give an error.

#rnx

- Running the virtualization process on a project without virtualization enabled will mount a 'FNullVirtualizationSystem' which will return errors when actually used.
- Now, once we initialize the virtualization system we check to see if it is actually enabled and if not we early out at this point. It's quicker and does not result in an error being logged.

diff --git a/Engine/Source/Programs/UnrealVirtualizationTool/Private/Commands/VirtualizeCommand.cpp b/Engine/Source/Programs/UnrealVirtualizationTool/Private/Commands/VirtualizeCommand.cpp
--- a/Engine/Source/Programs/UnrealVirtualizationTool/Private/Commands/VirtualizeCommand.cpp
+++ b/Engine/Source/Programs/UnrealVirtualizationTool/Private/Commands/VirtualizeCommand.cpp
@@ -174,6 +174,13 @@ bool FVirtualizeCommand::ProcessProject(const FProject& Project, TUniquePtr<FCom
 		UE::Virtualization::Shutdown();
 	};
 
+	if (!UE::Virtualization::IVirtualizationSystem::Get().IsEnabled())
+	{
+		UE_LOG(LogVirtualizationTool, Display, TEXT("\tVirtualization is not enabled for this project"));
+		Output = MakeUnique<FVirtualizeCommandOutput>(Project.GetProjectName(), TArray<FText>());
+		return true;
+	}
+
 	TArray<FString> ProjectPackages = Project.GetAllPackages();
 
 	EVirtualizationOptions Options = EVirtualizationOptions::None;
