[UVT] Further reduce the size of the tool

#rnx

- Minimal increase in compile time (+1.5s for full rebuild) but a significant reduction in exe/pdb size.

diff --git a/Engine/Source/Programs/UnrealVirtualizationTool/UnrealVirtualizationTool.Target.cs b/Engine/Source/Programs/UnrealVirtualizationTool/UnrealVirtualizationTool.Target.cs
--- a/Engine/Source/Programs/UnrealVirtualizationTool/UnrealVirtualizationTool.Target.cs
+++ b/Engine/Source/Programs/UnrealVirtualizationTool/UnrealVirtualizationTool.Target.cs
@@ -30,6 +30,10 @@ public class UnrealVirtualizationToolTarget : TargetRules
 		// Enable Developer plugins
 		bCompileWithPluginSupport = true;
 
+		// Reduce the final exe/pdb size
+		WindowsPlatform.bStripUnreferencedSymbols = true;
+		WindowsPlatform.bMergeIdenticalCOMDATs = true;
+
 		bEnableTrace = true;
 		GlobalDefinitions.Add("UE_SUPPORT_FULL_PACKAGEPATH=1");
 	}
