[UVT] Fix a bug where the child process was launched with an incorrect stdin pipe attached.

#rnx

- The problem with the process set up is described properly in UE-208628, but the quick summary is that a lot of our code is creating a single pipe, then using the input as stdin and the output as stdout when launching a child process so if that process ever did read from stdin it would end up using it's own output. Then we are copy/pasting the same bug around the code base,
- Fixed the problem by adding a second pipe to FProcessPipes.

diff --git a/Engine/Source/Programs/UnrealVirtualizationTool/Private/ProcessUtilities.h b/Engine/Source/Programs/UnrealVirtualizationTool/Private/ProcessUtilities.h
--- a/Engine/Source/Programs/UnrealVirtualizationTool/Private/ProcessUtilities.h
+++ b/Engine/Source/Programs/UnrealVirtualizationTool/Private/ProcessUtilities.h
@@ -16,29 +16,21 @@ struct FProcessPipes
 {
 	FProcessPipes()
 	{
-		verify(FPlatformProcess::CreatePipe(StdOut, StdIn));
+		verify(FPlatformProcess::CreatePipe(StdOutReadPipe, StdOutWritePipe, false));
+		verify(FPlatformProcess::CreatePipe(StdInReadPipe, StdInWritePipe, true));
 	}
 
 	~FProcessPipes()
 	{
-		FPlatformProcess::ClosePipe(StdOut, StdIn);
-	}
-
-	void* GetStdIn() const
-	{
-		return StdIn;
-	}
-
-	void* GetStdOut() const
-	{
-		return StdOut;
+		FPlatformProcess::ClosePipe(StdOutReadPipe, StdOutWritePipe);
+		FPlatformProcess::ClosePipe(StdInReadPipe, StdInWritePipe);
 	}
 
 	void ProcessStdOut()
 	{
-		check(StdOut != nullptr);
+		check(GetStdOutForReading() != nullptr);
 
-		FString Output = FPlatformProcess::ReadPipe(StdOut);
+		FString Output = FPlatformProcess::ReadPipe(GetStdOutForReading());
 
 		while (!Output.IsEmpty())
 		{
@@ -50,13 +42,37 @@ struct FProcessPipes
 				UE_LOG(LogVirtualizationTool, Display, TEXT("Child Process-> %s"), *Line);
 			}
 
-			Output = FPlatformProcess::ReadPipe(StdOut);
+			Output = FPlatformProcess::ReadPipe(GetStdOutForReading());
 		}
 	}
 
+	void* GetStdInForProcess() const
+	{
+		return StdInReadPipe;
+	}
+
+	void* GetStdInForWriting() const
+	{
+		return StdInWritePipe;
+	}
+
+	void* GetStdOutForProcess() const
+	{
+		return StdOutWritePipe;
+	}
+
+	void* GetStdOutForReading() const
+	{
+		return StdOutReadPipe;
+	}
+
 private:
-	void* StdIn = nullptr;
-	void* StdOut = nullptr;
+
+	void* StdOutReadPipe = nullptr;
+	void* StdOutWritePipe = nullptr;
+
+	void* StdInReadPipe = nullptr;
+	void* StdInWritePipe = nullptr;
 };
 
 } // namespace UE::Virtualization
\ No newline at end of file
diff --git a/Engine/Source/Programs/UnrealVirtualizationTool/Private/UnrealVirtualizationToolApp.cpp b/Engine/Source/Programs/UnrealVirtualizationTool/Private/UnrealVirtualizationToolApp.cpp
--- a/Engine/Source/Programs/UnrealVirtualizationTool/Private/UnrealVirtualizationToolApp.cpp
+++ b/Engine/Source/Programs/UnrealVirtualizationTool/Private/UnrealVirtualizationToolApp.cpp
@@ -836,13 +836,24 @@ bool FUnrealVirtualizationToolApp::LaunchChildProcess(const FCommand& Command, c
 		const bool bLaunchDetached = false;
 		const bool bLaunchHidden = true;
 		const bool bLaunchReallyHidden = true;
-		uint32 ProcessId = INDEX_NONE;
+
 		const int32 Priority = 0;
 		const TCHAR* WorkingDirectory = nullptr;
 
 		FProcessPipes Pipes;
-		FProcHandle Handle = FPlatformProcess::CreateProc(*CurrentExePath, *Args, bLaunchDetached, bLaunchHidden, bLaunchReallyHidden,
-			/*OutProcessID*/ nullptr, Priority, WorkingDirectory, Pipes.GetStdIn(), Pipes.GetStdOut());
+		FProcHandle Handle = FPlatformProcess::CreateProc
+		(
+			*CurrentExePath,
+			*Args,
+			bLaunchDetached,
+			bLaunchHidden,
+			bLaunchReallyHidden,
+			/*OutProcessID*/ nullptr,
+			Priority,
+			WorkingDirectory,
+			Pipes.GetStdOutForProcess(),
+			Pipes.GetStdInForProcess()
+		);
 
 		{
 			TRACE_CPUPROFILER_EVENT_SCOPE(WaitOnChildProcess);
