Improve the initialization logging when starting up UnrealVirtualizationTool.

#rnx

- Log which perforce server the tool connects too as sometimes it is not what we expect.
- Expanded code comments explaining why the CommandBase has two different source control provider pointer/references.

diff --git a/Engine/Source/Programs/UnrealVirtualizationTool/Private/Commands/CommandBase.cpp b/Engine/Source/Programs/UnrealVirtualizationTool/Private/Commands/CommandBase.cpp
--- a/Engine/Source/Programs/UnrealVirtualizationTool/Private/Commands/CommandBase.cpp
+++ b/Engine/Source/Programs/UnrealVirtualizationTool/Private/Commands/CommandBase.cpp
@@ -123,14 +123,11 @@ bool FCommand::TryConnectToSourceControl(FStringView ClientSpecName)
 		// If the project has been set then we can use the default source control provider
 		if (!ISourceControlModule::Get().GetProvider().IsEnabled())
 		{
-			// TODO - Warning
 			ISourceControlModule::Get().SetProvider(FName("Perforce"));
 		}
 
 		SCCProvider = &ISourceControlModule::Get().GetProvider();
 		SCCProvider->Init(true);
-
-		return SCCProvider->IsAvailable();
 	}
 	else
 	{
@@ -146,24 +143,30 @@ bool FCommand::TryConnectToSourceControl(FStringView ClientSpecName)
 		{
 			SCCProvider = OwnedSCCProvider.Get();
 			SCCProvider->Init(true);
-
-			if (SCCProvider->IsAvailable())
-			{
-				return true;
-			}
-			else
-			{
-				UE_LOG(LogVirtualizationTool, Error, TEXT("Failed to establish a perforce connection"));
-				return false;
-			}
-
 		}
 		else
 		{
-			UE_LOG(LogVirtualizationTool, Error, TEXT("Failed to instantiate a perforce revision control connection"));
+			UE_LOG(LogVirtualizationTool, Error, TEXT("\tFailed to instantiate a perforce revision control connection"));
 			return false;
 		}
 	}
+
+	if (SCCProvider->IsAvailable())
+	{
+		TMap<ISourceControlProvider::EStatus, FString> StatusMap = SCCProvider->GetStatus();
+
+		if (FString* Server = StatusMap.Find(ISourceControlProvider::EStatus::Port))
+		{
+			UE_LOG(LogVirtualizationTool, Display, TEXT("\tSuccessfully connected to server '%s'"), **Server);
+		}
+		
+		return true;
+	}
+	else
+	{
+		UE_LOG(LogVirtualizationTool, Error, TEXT("\tFailed to establish a perforce connection"));
+		return false;
+	}
 }
 
 bool FCommand::TryParseChangelist(FStringView ClientSpecName, FStringView ChangelistNumber, TArray<FString>& OutPackages, FSourceControlChangelistPtr* OutChangelist)
diff --git a/Engine/Source/Programs/UnrealVirtualizationTool/Private/Commands/CommandBase.h b/Engine/Source/Programs/UnrealVirtualizationTool/Private/Commands/CommandBase.h
--- a/Engine/Source/Programs/UnrealVirtualizationTool/Private/Commands/CommandBase.h
+++ b/Engine/Source/Programs/UnrealVirtualizationTool/Private/Commands/CommandBase.h
@@ -141,10 +141,19 @@ protected: // Common SourceControl Code
 private:
 	FString CommandName;
 
+	/*
+	 * Sometimes the command will have to create and own the source control provider and be responsible for
+	 * cleaning it up once the command has finished. In these cases we store the pointer to the provider
+	 * that we own, here.
+	 */
 	TUniquePtr<ISourceControlProvider> OwnedSCCProvider;
 
 protected:
 
+	/*
+	 * Pointer to the source control provider that the command should use.This can either be to a provider
+	 * that the command owns in 'OwnedSCCProvider' or one provided for us by the source control api.
+	 */
 	ISourceControlProvider* SCCProvider = nullptr;
 
 	
