UnrealVirtualizationTool now accepts the P4Port and P4User values on the commandline.

#rnx

- Clean up some of the logging to be tab aligned.
- Add an additional line of logging showing how many files were found in the changelist. Previously I was only logging how many packages (.umap/.uasset) were found in the changelist

diff --git a/Engine/Source/Programs/UnrealVirtualizationTool/Private/Commands/CommandBase.cpp b/Engine/Source/Programs/UnrealVirtualizationTool/Private/Commands/CommandBase.cpp
--- a/Engine/Source/Programs/UnrealVirtualizationTool/Private/Commands/CommandBase.cpp
+++ b/Engine/Source/Programs/UnrealVirtualizationTool/Private/Commands/CommandBase.cpp
@@ -136,6 +136,7 @@ bool FCommand::TryConnectToSourceControl(FStringView ClientSpecName)
 		// we are parsing a perforce changelist for files to operate on.
 		FSourceControlInitSettings SCCSettings(FSourceControlInitSettings::EBehavior::OverrideAll);
 		SCCSettings.SetConfigBehavior(FSourceControlInitSettings::EConfigBehavior::ReadOnly);
+		SCCSettings.SetCmdLineFlags(FSourceControlInitSettings::ECmdLineFlags::ReadAll);
 		SCCSettings.AddSetting(TEXT("P4Client"), ClientSpecName);
 
 		OwnedSCCProvider = ISourceControlModule::Get().CreateProvider(FName("Perforce"), TEXT("UnrealVirtualizationTool"), SCCSettings);
@@ -231,6 +232,9 @@ bool FCommand::TryParseChangelist(FStringView ClientSpecName, FStringView Change
 			}
 
 			const TArray<FSourceControlStateRef>& FilesinChangelist = ChangelistState->GetFilesStates();
+
+			UE_LOG(LogVirtualizationTool, Log, TEXT("\tFound %d files in the changelist"), FilesinChangelist.Num());
+
 			for (const FSourceControlStateRef& FileState : FilesinChangelist)
 			{
 				if (IsPackageFile(FileState->GetFilename()))
@@ -239,7 +243,7 @@ bool FCommand::TryParseChangelist(FStringView ClientSpecName, FStringView Change
 				}
 				else
 				{
-					UE_LOG(LogVirtualizationTool, Log, TEXT("Ignoring non-package file '%s'"), *FileState->GetFilename());
+					UE_LOG(LogVirtualizationTool, Log, TEXT("\tIgnoring non-package file '%s'"), *FileState->GetFilename());
 				}
 			}
 
