Adding internal API to allow monolithic programs to optionally accept a uproject as a first argument and use that as its project directory even if it is an engine executable.

diff --git a/Engine/Source/Programs/DerivedDataTool/Private/DerivedDataTool.cpp b/Engine/Source/Programs/DerivedDataTool/Private/DerivedDataTool.cpp
--- a/Engine/Source/Programs/DerivedDataTool/Private/DerivedDataTool.cpp
+++ b/Engine/Source/Programs/DerivedDataTool/Private/DerivedDataTool.cpp
@@ -7,6 +7,7 @@
 #include "HAL/LowLevelMemTracker.h"
 #include "Misc/ScopeExit.h"
 #include "Modules/ModuleManager.h"
+#include "ProjectUtilities.h"
 
 DEFINE_LOG_CATEGORY_STATIC(LogDerivedDataTool, Log, All);
 
@@ -18,6 +19,9 @@ INT32_MAIN_INT32_ARGC_TCHAR_ARGV()
 
 	const FTaskTagScope Scope(ETaskTag::EGameThread);
 
+	// Allows this program to accept a project argument on the commandline and use project-specific config
+	UE::ProjectUtilities::ParseProjectDirFromCommandline(ArgC, ArgV);
+
 	if (int32 Ret = GEngineLoop.PreInit(ArgC, ArgV))
 	{
 		return Ret;
diff --git a/Engine/Source/Programs/UnrealVirtualizationTool/Private/UnrealVirtualizationTool.cpp b/Engine/Source/Programs/UnrealVirtualizationTool/Private/UnrealVirtualizationTool.cpp
--- a/Engine/Source/Programs/UnrealVirtualizationTool/Private/UnrealVirtualizationTool.cpp
+++ b/Engine/Source/Programs/UnrealVirtualizationTool/Private/UnrealVirtualizationTool.cpp
@@ -3,6 +3,7 @@
 #include "UnrealVirtualizationTool.h"
 
 #include "Modules/ModuleManager.h"
+#include "ProjectUtilities.h"
 #include "RequiredProgramMainCPPInclude.h"
 #include "UnrealVirtualizationToolApp.h"
 
@@ -16,32 +17,8 @@ int32 UnrealVirtualizationToolMain(int32 ArgC, TCHAR* ArgV[])
 
 	using namespace UE::Virtualization;
 
-	// Although standalone tools can set the project path via the cmdline this will not change the ProjectDir being
-	// used as standalone tools have a beespoke path in FGenericPlatformMisc::ProjectDir. We can work around this
-	// for now by doing our own parsing then using the project dir override feature.
-	// This is a band aid while we consider adding better support for project files/directories with stand alone tools.
-	if(ArgC >= 2)
-	{
-		FString Cmd(ArgV[1]);
-
-		if (!Cmd.IsEmpty() && !Cmd.StartsWith(TEXT("-")) && Cmd.EndsWith(FProjectDescriptor::GetExtension()))
-		{
-			FString ProjectDir = FPaths::GetPath(Cmd);
-			ProjectDir = FFileManagerGeneric::DefaultConvertToRelativePath(*ProjectDir);
-			
-			// The path should end with a trailing slash (see FGenericPlatformMisc::ProjectDir) so we use
-			// NormalizeFilename not NormalizeDirectoryName as the latter will remove trailing slashes. We
-			// also need to add one if it is missing.
-			// We probably should move this path fixup code to 'FPlatformMisc::SetOverrideProjectDir'
-			FPaths::NormalizeFilename(ProjectDir);
-			if (!ProjectDir.EndsWith(TEXT("/")))
-			{
-				ProjectDir += TEXT("/");
-			}
-
-			FPlatformMisc::SetOverrideProjectDir(ProjectDir);
-		}
-	}
+	// Allows this program to accept a project argument on the commandline and use project-specific config
+	UE::ProjectUtilities::ParseProjectDirFromCommandline(ArgC, ArgV);
 
 	GEngineLoop.PreInit(ArgC, ArgV);
 	check(GConfig && GConfig->IsReadyForUse());
diff --git a/Engine/Source/Programs/ZenLaunch/Private/ZenLaunch.cpp b/Engine/Source/Programs/ZenLaunch/Private/ZenLaunch.cpp
--- a/Engine/Source/Programs/ZenLaunch/Private/ZenLaunch.cpp
+++ b/Engine/Source/Programs/ZenLaunch/Private/ZenLaunch.cpp
@@ -8,6 +8,7 @@
 #include "Misc/CommandLine.h"
 #include "Misc/Parse.h"
 #include "String/LexFromString.h"
+#include "ProjectUtilities.h"
 
 #include "RequiredProgramMainCPPInclude.h"
 
@@ -27,32 +28,8 @@ INT32_MAIN_INT32_ARGC_TCHAR_ARGV()
 		FEngineLoop::AppExit();
 	};
 
-	// Although standalone tools can set the project path via the cmdline this will not change the ProjectDir being
-	// used as standalone tools have a beespoke path in FGenericPlatformMisc::ProjectDir. We can work around this
-	// for now by doing our own parsing then using the project dir override feature.
-	// This is a band aid while we consider adding better support for project files/directories with stand alone tools.
-	if (ArgC >= 2)
-	{
-		FString Cmd(ArgV[1]);
-
-		if (!Cmd.IsEmpty() && !Cmd.StartsWith(TEXT("-")) && Cmd.EndsWith(FProjectDescriptor::GetExtension()))
-		{
-			FString ProjectDir = FPaths::GetPath(Cmd);
-			ProjectDir = FFileManagerGeneric::DefaultConvertToRelativePath(*ProjectDir);
-
-			// The path should end with a trailing slash (see FGenericPlatformMisc::ProjectDir) so we use
-			// NormalizeFilename not NormalizeDirectoryName as the latter will remove trailing slashes. We
-			// also need to add one if it is missing.
-			// We probably should move this path fixup code to 'FPlatformMisc::SetOverrideProjectDir'
-			FPaths::NormalizeFilename(ProjectDir);
-			if (!ProjectDir.EndsWith(TEXT("/")))
-			{
-				ProjectDir += TEXT("/");
-			}
-
-			FPlatformMisc::SetOverrideProjectDir(ProjectDir);
-		}
-	}
+	// Allows this program to accept a project argument on the commandline and use project-specific config
+	UE::ProjectUtilities::ParseProjectDirFromCommandline(ArgC, ArgV);
 
 	if (int32 Ret = GEngineLoop.PreInit(ArgC, ArgV))
 	{
diff --git a/Engine/Source/Runtime/Projects/Internal/ProjectUtilities.h b/Engine/Source/Runtime/Projects/Internal/ProjectUtilities.h
new file mode 100644
--- /dev/null
+++ b/Engine/Source/Runtime/Projects/Internal/ProjectUtilities.h
@@ -0,0 +1,13 @@
+// Copyright Epic Games, Inc. All Rights Reserved.
+
+#pragma once
+
+#include "CoreTypes.h"
+
+namespace UE::ProjectUtilities
+{
+	/**
+	 * Allows a monolithic target to control the game directory by the presence of commandline arguments for the project.
+	 */
+	PROJECTS_API void ParseProjectDirFromCommandline(int32 ArgC, TCHAR* ArgV[]);
+}
diff --git a/Engine/Source/Runtime/Projects/Private/ProjectUtilities.cpp b/Engine/Source/Runtime/Projects/Private/ProjectUtilities.cpp
new file mode 100644
--- /dev/null
+++ b/Engine/Source/Runtime/Projects/Private/ProjectUtilities.cpp
@@ -0,0 +1,41 @@
+// Copyright Epic Games, Inc. All Rights Reserved.
+
+#include "ProjectUtilities.h"
+#include "Containers/UnrealString.h"
+#include "HAL/FileManagerGeneric.h"
+#include "HAL/PlatformMisc.h"
+#include "Misc/Paths.h"
+#include "ProjectDescriptor.h"
+
+namespace UE::ProjectUtilities
+{
+
+void ParseProjectDirFromCommandline(int32 ArgC, TCHAR* ArgV[])
+{
+	// Although standalone tools can set the project path via the cmdline this will not change the ProjectDir being
+	// used as standalone tools have a bespoke path in FGenericPlatformMisc::ProjectDir. We can address this
+	// by doing our own parsing then using the project dir override feature.
+	if (ArgC >= 2)
+	{
+		FString Cmd(ArgV[1]);
+
+		if (!Cmd.IsEmpty() && !Cmd.StartsWith(TEXT("-")) && Cmd.EndsWith(FProjectDescriptor::GetExtension()))
+		{
+			FString ProjectDir = FPaths::GetPath(Cmd);
+			ProjectDir = FFileManagerGeneric::DefaultConvertToRelativePath(*ProjectDir);
+
+			// The path should end with a trailing slash (see FGenericPlatformMisc::ProjectDir) so we use
+			// NormalizeFilename not NormalizeDirectoryName as the latter will remove trailing slashes. We
+			// also need to add one if it is missing.
+			FPaths::NormalizeFilename(ProjectDir);
+			if (!ProjectDir.EndsWith(TEXT("/")))
+			{
+				ProjectDir += TEXT("/");
+			}
+
+			FPlatformMisc::SetOverrideProjectDir(ProjectDir);
+		}
+	}
+}
+
+}
