[UVT] Clean up the child process files after use.

#rnx

- When the UVT needs to launch a child process we send it info on the command to run via a .input file and in turn the child process will write out a .output file to return info to the calling process.
- This change will now attempt to delete these files after we are done with them.
- If the delete fails we will raise a warning rather than an error as it is not essential that the files be deleted.

diff --git a/Engine/Source/Programs/UnrealVirtualizationTool/Private/UnrealVirtualizationToolApp.cpp b/Engine/Source/Programs/UnrealVirtualizationTool/Private/UnrealVirtualizationToolApp.cpp
--- a/Engine/Source/Programs/UnrealVirtualizationTool/Private/UnrealVirtualizationToolApp.cpp
+++ b/Engine/Source/Programs/UnrealVirtualizationTool/Private/UnrealVirtualizationToolApp.cpp
@@ -18,9 +18,11 @@
 #include "Misc/Parse.h"
 #include "Misc/Paths.h"
 #include "Misc/PathViews.h"
+#include "Misc/ScopeExit.h"
 #include "Modules/ModuleManager.h"
 #include "ProcessUtilities.h"
 #include "UnrealVirtualizationTool.h"
+#include "VirtualizationUtilities.h"
 #include "Virtualization/VirtualizationSystem.h"
 
 namespace
@@ -806,6 +808,29 @@ bool FUnrealVirtualizationToolApp::TryWriteChildProcessInputFile(const FGuid& Ch
 	return true;
 }
 
+void FUnrealVirtualizationToolApp::CleanUpChildProcessFiles(const FGuid& ChildProcessId)
+{
+	// Note: A better way to do this would be FILE_FLAG_DELETE_ON_CLOSE  so that the files
+	// are cleaned up when this process is destroyed but we do not currently expose this
+	// sort of functionality.
+
+	const TCHAR* FileExtensions[] = { TEXT("input"), TEXT("output") };
+
+	for (const TCHAR* Extension : FileExtensions)
+	{
+		TStringBuilder<512> FilePath;
+		CreateChildProcessFilePath(ChildProcessId, Extension, FilePath);
+
+		if (!IFileManager::Get().Delete(FilePath.ToString()))
+		{
+			TStringBuilder<MAX_SPRINTF> SystemErrorMsg;
+			UE::Virtualization::Utils::GetFormattedSystemError(SystemErrorMsg);
+
+			UE_LOG(LogVirtualizationTool, Warning, TEXT("Failed to clean up temp file '%s' due to: %s"), FilePath.ToString(), SystemErrorMsg.ToString());
+		}
+	}
+}
+
 bool FUnrealVirtualizationToolApp::LaunchChildProcess(const FCommand& Command, const FProject& Project, FStringView GlobalOptions, TArray<TUniquePtr<FCommandOutput>>& OutputArray)
 {
 	TRACE_CPUPROFILER_EVENT_SCOPE(LaunchChildProcess);
@@ -821,6 +846,11 @@ bool FUnrealVirtualizationToolApp::LaunchChildProcess(const FCommand& Command, c
 		return false;
 	}
 
+	ON_SCOPE_EXIT
+	{
+		CleanUpChildProcessFiles(ChildProcessId);
+	};
+
 	{
 		TRACE_CPUPROFILER_EVENT_SCOPE(RunChildProcess);
 
diff --git a/Engine/Source/Programs/UnrealVirtualizationTool/Private/UnrealVirtualizationToolApp.h b/Engine/Source/Programs/UnrealVirtualizationTool/Private/UnrealVirtualizationToolApp.h
--- a/Engine/Source/Programs/UnrealVirtualizationTool/Private/UnrealVirtualizationToolApp.h
+++ b/Engine/Source/Programs/UnrealVirtualizationTool/Private/UnrealVirtualizationToolApp.h
@@ -83,6 +83,7 @@ private:
 	/** Note not currently static due to dependencies */
 	bool TryReadChildProcessInputFile(const FString& InputPath);
 	static bool TryWriteChildProcessInputFile(const FGuid& ChildProcessId, const FCommand& Command, const FProject& Project, FStringBuilderBase& OutPath);
+	static void CleanUpChildProcessFiles(const FGuid& ChildProcessId);
 
 
 	static bool LaunchChildProcess(const FCommand& Command, const FProject& Project, FStringView GlobalOptions, TArray<TUniquePtr<FCommandOutput>>& OutputArray);
