Fix a bug where UVT invoking a child process would fail if the engine root path or the target project path has spaces in it.

- When invoking a child process UVT will pass in the path of the target .uproject file as well as a path to a .inputfile containing the details of the operation to run. If either of these paths contained whitespace then they would fail to parse correctly and either the target project would not be correctly set or  the input file would not be readable.
- We now properly escape the args so that the whole paths are parsed, even when containing whitespace.

diff --git a/Engine/Source/Programs/UnrealVirtualizationTool/Private/UnrealVirtualizationToolApp.cpp b/Engine/Source/Programs/UnrealVirtualizationTool/Private/UnrealVirtualizationToolApp.cpp
--- a/Engine/Source/Programs/UnrealVirtualizationTool/Private/UnrealVirtualizationToolApp.cpp
+++ b/Engine/Source/Programs/UnrealVirtualizationTool/Private/UnrealVirtualizationToolApp.cpp
@@ -855,7 +855,7 @@ bool FUnrealVirtualizationToolApp::LaunchChildProcess(const FCommand& Command, c
 		TRACE_CPUPROFILER_EVENT_SCOPE(RunChildProcess);
 
 		const FString CurrentExePath = FPlatformProcess::ExecutablePath();
-		FString Args = FString::Printf(TEXT("%s -ChildProcess=%s"), *Project.GetProjectFilePath(), InputFilePath.ToString());
+		FString Args = FString::Printf(TEXT("\"%s\" -ChildProcess=\"%s\""), *Project.GetProjectFilePath(), InputFilePath.ToString());
 
 		if (!GlobalOptions.IsEmpty())
 		{
