Strip out additional features from UnrealVirtualizationTool that it does not need.

#rnx

diff --git a/Engine/Source/Programs/UnrealVirtualizationTool/UnrealVirtualizationTool.Target.cs b/Engine/Source/Programs/UnrealVirtualizationTool/UnrealVirtualizationTool.Target.cs
--- a/Engine/Source/Programs/UnrealVirtualizationTool/UnrealVirtualizationTool.Target.cs
+++ b/Engine/Source/Programs/UnrealVirtualizationTool/UnrealVirtualizationTool.Target.cs
@@ -33,6 +33,8 @@ public class UnrealVirtualizationToolTarget : TargetRules
 		// Reduce the final exe/pdb size
 		WindowsPlatform.bStripUnreferencedSymbols = true;
 		WindowsPlatform.bMergeIdenticalCOMDATs = true;
+		WindowsPlatform.bUseBundledDbgHelp = false;
+		WindowsPlatform.bPixProfilingEnabled = false;
 
 		bEnableTrace = true;
 		GlobalDefinitions.Add("UE_SUPPORT_FULL_PACKAGEPATH=1");
