Add a cmdline arg '-ReportFailures' to UnrealVirtualizationTool that raises an ensure if the process fails.

#rnx

- At the moment most projects allow submits with failed virtualization attempts to continue uninterrupted but that means we rely on users reporting issues to us. By raising an ensure we will get a report sent to our crash reporting systems which we can use to dig into the problems. This system will be opt in via the cmdline arg.
- The cmdline arg will be passed onto any spawned child processes so that they will report too.
-- Note that we also now pass '-fastexit' onto the child processes which was not done before.
- We only raise an ensure if the problem occurred in the current process, it is expected that child processes will raise their own.
-- We use a new enum EProcessResult to pass this info around.
- The log filtering has been extended to filter out warning/error messages normally logged when an ensure fires as these will concern the end users, we don't want the reporting to not interrupt their day to day work either.

diff --git a/Engine/Source/Programs/UnrealVirtualizationTool/Private/UnrealVirtualizationTool.cpp b/Engine/Source/Programs/UnrealVirtualizationTool/Private/UnrealVirtualizationTool.cpp
--- a/Engine/Source/Programs/UnrealVirtualizationTool/Private/UnrealVirtualizationTool.cpp
+++ b/Engine/Source/Programs/UnrealVirtualizationTool/Private/UnrealVirtualizationTool.cpp
@@ -23,6 +23,8 @@ int32 UnrealVirtualizationToolMain(int32 ArgC, TCHAR* ArgV[])
 	GEngineLoop.PreInit(ArgC, ArgV);
 	check(GConfig && GConfig->IsReadyForUse());
 
+	const bool bReportFailures = FParse::Param(FCommandLine::Get(), TEXT("ReportFailures"));
+	
 #if 0
 	while (!FPlatformMisc::IsDebuggerPresent())
 	{
@@ -34,28 +36,34 @@ int32 UnrealVirtualizationToolMain(int32 ArgC, TCHAR* ArgV[])
 
 	FModuleManager::Get().StartProcessingNewlyLoadedObjects();
 
-	bool bRanSuccessfully = true;
+	EProcessResult ProcessResult = EProcessResult::Success;
 
 	FUnrealVirtualizationToolApp App;
 
-	EInitResult Result = App.Initialize();
-	if (Result == EInitResult::Success)
+	EInitResult InitResult = App.Initialize();
+	if (InitResult == EInitResult::Success)
 	{
-		if (!App.Run())
+		ProcessResult = App.Run();
+		if (ProcessResult != EProcessResult::Success)
 		{
 			UE_LOG(LogVirtualizationTool, Error, TEXT("UnrealVirtualizationTool ran with errors"));
-			bRanSuccessfully = false;
 		}
 	}	
-	else if(Result == EInitResult::Error)
+	else if(InitResult == EInitResult::Error)
 	{
 		UE_LOG(LogVirtualizationTool, Error, TEXT("UnrealVirtualizationTool failed to initialize"));
-		bRanSuccessfully = false;
+		ProcessResult = EProcessResult::Error;
 	}
 
-	UE_CLOG(bRanSuccessfully, LogVirtualizationTool, Display, TEXT("UnrealVirtualizationTool ran successfully"));
+	UE_CLOG(ProcessResult == EProcessResult::Success, LogVirtualizationTool, Display, TEXT("UnrealVirtualizationTool ran successfully"));
 
-	const uint8 ReturnCode = bRanSuccessfully ? 0 : 1;
+	// Don't report if the error was in a child process, they will raise their own ensures
+	if (bReportFailures && ProcessResult == EProcessResult::Error)
+	{
+		ensure(false);
+	}
+
+	const uint8 ReturnCode = ProcessResult == EProcessResult::Success ? 0 : 1;
 
 	if (FParse::Param(FCommandLine::Get(), TEXT("fastexit")))
 	{
diff --git a/Engine/Source/Programs/UnrealVirtualizationTool/Private/UnrealVirtualizationToolApp.cpp b/Engine/Source/Programs/UnrealVirtualizationTool/Private/UnrealVirtualizationToolApp.cpp
--- a/Engine/Source/Programs/UnrealVirtualizationTool/Private/UnrealVirtualizationToolApp.cpp
+++ b/Engine/Source/Programs/UnrealVirtualizationTool/Private/UnrealVirtualizationToolApp.cpp
@@ -135,11 +135,11 @@ public:
 	virtual void Serialize(const TCHAR* V, ELogVerbosity::Type Verbosity, const FName& Category, double Time) override
 	{
 #if !NO_LOGGING
-		if (Verbosity == ELogVerbosity::Display && Category != LogVirtualizationTool.GetCategoryName())
-#endif
+		if (ShouldFilterMessage(Verbosity, Category))
 		{
 			Verbosity = ELogVerbosity::Log;
 		}
+#endif // !NO_LOGGING
 
 		FFeedbackContextAnsi::Serialize(V, Verbosity, Category, Time);
 	}
@@ -147,18 +147,38 @@ public:
 	virtual void SerializeRecord(const UE::FLogRecord& Record) override
 	{
 #if !NO_LOGGING
-		if (Record.GetVerbosity() == ELogVerbosity::Display && Record.GetCategory() != LogVirtualizationTool.GetCategoryName())
+		if (ShouldFilterMessage(Record.GetVerbosity(), Record.GetCategory()))
 		{
 			UE::FLogRecord LocalRecord = Record;
 			LocalRecord.SetVerbosity(ELogVerbosity::Log);
 			return FFeedbackContextAnsi::SerializeRecord(LocalRecord);
 		}
-#endif
+#endif // !NO_LOGGING
 
 		FFeedbackContextAnsi::SerializeRecord(Record);
 	}
 
 private:
+
+	bool ShouldFilterMessage(ELogVerbosity::Type Verbosity, const FName& Category)
+	{
+#if !NO_LOGGING
+		// We only want 'LogVirtualizationTool' messages in display
+		if (Verbosity == ELogVerbosity::Display && Category != LogVirtualizationTool.GetCategoryName())
+		{
+			return true;
+		}
+
+		// Suppress errors from our reporting systems
+		if (Verbosity == ELogVerbosity::Error && Category == LogOutputDevice.GetCategoryName())
+		{
+			return true;
+		}
+#endif //  !NO_LOGGING
+
+		return false;
+	}
+
 	FFeedbackContext* OriginalLog;
 };
 
@@ -223,15 +243,17 @@ EInitResult FUnrealVirtualizationToolApp::Initialize()
 	return EInitResult::Success;
 }
 
-bool FUnrealVirtualizationToolApp::Run()
+EProcessResult FUnrealVirtualizationToolApp::Run()
 {
 	TRACE_CPUPROFILER_EVENT_SCOPE(Run);
 
 	TArray<TUniquePtr<FCommandOutput>> OutputArray;
-	if (!ProcessProjects(OutputArray))
+
+	EProcessResult Result = ProcessProjects(OutputArray);
+	if (Result != EProcessResult::Success)
 	{
 		UE_LOG(LogVirtualizationTool, Error, TEXT("Command '%s' failed!"), *CurrentCommand->GetName());
-		return false;
+		return Result;
 	}
 
 	if (!IsChildProcess())
@@ -239,7 +261,7 @@ bool FUnrealVirtualizationToolApp::Run()
 		if (!CurrentCommand->ProcessOutput(OutputArray))
 		{
 			UE_LOG(LogVirtualizationTool, Error, TEXT("Command '%s' failed!"), *CurrentCommand->GetName());
-			return false;
+			return EProcessResult::Error;
 		}
 	}
 	else
@@ -247,16 +269,16 @@ bool FUnrealVirtualizationToolApp::Run()
 		if (!TryWriteChildProcessOutputFile(ChildProcessId, OutputArray))
 		{
 			UE_LOG(LogVirtualizationTool, Error, TEXT("Command '%s' failed!"), *CurrentCommand->GetName());
-			return false;
+			return EProcessResult::Error;
 		}
 	}
 
 
 	UE_LOG(LogVirtualizationTool, Display, TEXT("Command '%s' succeeded!"), *CurrentCommand->GetName());
-	return true;
+	return EProcessResult::Success;
 }
 
-bool FUnrealVirtualizationToolApp::ProcessProjects(TArray<TUniquePtr<FCommandOutput>>& OutputArray)
+EProcessResult FUnrealVirtualizationToolApp::ProcessProjects(TArray<TUniquePtr<FCommandOutput>>& OutputArray)
 {
 	TRACE_CPUPROFILER_EVENT_SCOPE(ProcessProjects);
 
@@ -276,8 +298,8 @@ bool FUnrealVirtualizationToolApp::ProcessProjects(TArray<TUniquePtr<FCommandOut
 			{
 				// If we are a child process then the correct project path should have been provided and so
 				// this check is mostly a paranoid check to make sure things are working as we expect.
-				UE_LOG(LogVirtualizationTool, Error, TEXT("The child process was create with project path '%s' but expected '%s'"), *ProjectPath, *Project.GetProjectFilePath());
-				return false;
+				UE_LOG(LogVirtualizationTool, Error, TEXT("The child process was created with project path '%s' but expected '%s'"), *ProjectPath, *Project.GetProjectFilePath());
+				return EProcessResult::Error;
 			}
 
 			check(ProjectPath == Project.GetProjectFilePath());
@@ -285,7 +307,7 @@ bool FUnrealVirtualizationToolApp::ProcessProjects(TArray<TUniquePtr<FCommandOut
 			TUniquePtr<FCommandOutput> Output;
 			if (!CurrentCommand->ProcessProject(Project, Output))
 			{
-				return false;
+				return EProcessResult::Error;
 			}
 
 			if (Output != nullptr)
@@ -295,14 +317,15 @@ bool FUnrealVirtualizationToolApp::ProcessProjects(TArray<TUniquePtr<FCommandOut
 		}
 		else
 		{
-			if (!LaunchChildProcess(*CurrentCommand, Project, GlobalCmdlineOptions, OutputArray))
+			EProcessResult Result = LaunchChildProcess(*CurrentCommand, Project, GlobalCmdlineOptions, OutputArray);
+			if (Result != EProcessResult::Success)
 			{
-				return false;
+				return Result;
 			}
 		}
 	}
 
-	return true;
+	return EProcessResult::Success;
 }
 
 void FUnrealVirtualizationToolApp::PrintCmdLineHelp() const
@@ -433,6 +456,12 @@ EInitResult FUnrealVirtualizationToolApp::TryParseGlobalOptions(const TCHAR* Cmd
 		OutputDeviceOverride.Reset();
 	}
 
+	// Now add commandline switches used in UnrealVirtualizationToolMain (probably should be doing this setup work there)
+	if (FParse::Param(CmdLine, TEXT("ReportFailures")))
+	{
+		AddGlobalOption(TEXT("-ReportFailures"));
+	}
+
 	return EInitResult::Success;
 }
 
@@ -831,7 +860,7 @@ void FUnrealVirtualizationToolApp::CleanUpChildProcessFiles(const FGuid& ChildPr
 	}
 }
 
-bool FUnrealVirtualizationToolApp::LaunchChildProcess(const FCommand& Command, const FProject& Project, FStringView GlobalOptions, TArray<TUniquePtr<FCommandOutput>>& OutputArray)
+EProcessResult FUnrealVirtualizationToolApp::LaunchChildProcess(const FCommand& Command, const FProject& Project, FStringView GlobalOptions, TArray<TUniquePtr<FCommandOutput>>& OutputArray)
 {
 	TRACE_CPUPROFILER_EVENT_SCOPE(LaunchChildProcess);
 
@@ -843,7 +872,7 @@ bool FUnrealVirtualizationToolApp::LaunchChildProcess(const FCommand& Command, c
 	if (!TryWriteChildProcessInputFile(ChildProcessId, Command, Project, InputFilePath))
 	{
 		// No need to log an error here, ::TryWriteChildProcessInputFile will take care of that
-		return false;
+		return EProcessResult::Error;
 	}
 
 	ON_SCOPE_EXIT
@@ -898,23 +927,23 @@ bool FUnrealVirtualizationToolApp::LaunchChildProcess(const FCommand& Command, c
 		if (!FPlatformProcess::GetProcReturnCode(Handle, &ReturnCode))
 		{
 			UE_LOG(LogVirtualizationTool, Display, TEXT("Failed to retrieve the return value of the child process"));
-			return false;
+			return EProcessResult::Error;
 		}
 
 		if (ReturnCode != 0)
 		{
 			UE_LOG(LogVirtualizationTool, Display, TEXT("Child process failed with error code: %d"), ReturnCode);
-			return false;
+			return EProcessResult::ChildProcessError;
 		}
 	}
 
 	if (!TryReadChildProcessOutputFile(ChildProcessId, Command, OutputArray))
 	{
 		// No need to log an error here, ::TryReadChildProcessOutputFile will take care of that
-		return false;
+		return EProcessResult::Error;
 	}
 
-	return true;
+	return EProcessResult::Success;
 }
 
 } // namespace UE::Virtualization
diff --git a/Engine/Source/Programs/UnrealVirtualizationTool/Private/UnrealVirtualizationToolApp.h b/Engine/Source/Programs/UnrealVirtualizationTool/Private/UnrealVirtualizationToolApp.h
--- a/Engine/Source/Programs/UnrealVirtualizationTool/Private/UnrealVirtualizationToolApp.h
+++ b/Engine/Source/Programs/UnrealVirtualizationTool/Private/UnrealVirtualizationToolApp.h
@@ -46,6 +46,13 @@ enum class EInitResult
 	Error
 };
 
+enum class EProcessResult
+{
+	Success = 0,
+	Error,
+	ChildProcessError
+};
+
 class FUnrealVirtualizationToolApp
 {
 public:
@@ -53,13 +60,13 @@ public:
 	~FUnrealVirtualizationToolApp();
 
 	EInitResult Initialize();
-	bool Run();
+	EProcessResult Run();
 
 private:
 
 	void PrintCmdLineHelp() const;
 
-	bool ProcessProjects(TArray<TUniquePtr<FCommandOutput>>& OutputArray);
+	EProcessResult ProcessProjects(TArray<TUniquePtr<FCommandOutput>>& OutputArray);
 
 	bool TryLoadModules();
 	bool TryInitEnginePlugins();
@@ -85,8 +92,7 @@ private:
 	static bool TryWriteChildProcessInputFile(const FGuid& ChildProcessId, const FCommand& Command, const FProject& Project, FStringBuilderBase& OutPath);
 	static void CleanUpChildProcessFiles(const FGuid& ChildProcessId);
 
-
-	static bool LaunchChildProcess(const FCommand& Command, const FProject& Project, FStringView GlobalOptions, TArray<TUniquePtr<FCommandOutput>>& OutputArray);
+	static EProcessResult LaunchChildProcess(const FCommand& Command, const FProject& Project, FStringView GlobalOptions, TArray<TUniquePtr<FCommandOutput>>& OutputArray);
 
 private:
 
