UnrealVirtualizationTool can now pass the cmdline arg '-fastexit' to spawned child processes.

#rnx

diff --git a/Engine/Source/Programs/UnrealVirtualizationTool/Private/UnrealVirtualizationToolApp.cpp b/Engine/Source/Programs/UnrealVirtualizationTool/Private/UnrealVirtualizationToolApp.cpp
--- a/Engine/Source/Programs/UnrealVirtualizationTool/Private/UnrealVirtualizationToolApp.cpp
+++ b/Engine/Source/Programs/UnrealVirtualizationTool/Private/UnrealVirtualizationToolApp.cpp
@@ -462,6 +462,11 @@ EInitResult FUnrealVirtualizationToolApp::TryParseGlobalOptions(const TCHAR* Cmd
 		AddGlobalOption(TEXT("-ReportFailures"));
 	}
 
+	if (FParse::Param(CmdLine, TEXT("fastexit")))
+	{
+		AddGlobalOption(TEXT("-fastexit"));
+	}
+
 	return EInitResult::Success;
 }
 
