[UVT] Do not try to virtualize packages marked fro delete or ignore in revision control.

- Previously we passed all files in the changelist to the virtualization process which would result in errors when the process tried to open deleted packages for reading.
- We now skip packages marked for delete and packages marked for ignore, the latter won't be submitted so there is no point virtualizing them.

diff --git a/Engine/Source/Programs/UnrealVirtualizationTool/Private/Commands/CommandBase.cpp b/Engine/Source/Programs/UnrealVirtualizationTool/Private/Commands/CommandBase.cpp
--- a/Engine/Source/Programs/UnrealVirtualizationTool/Private/Commands/CommandBase.cpp
+++ b/Engine/Source/Programs/UnrealVirtualizationTool/Private/Commands/CommandBase.cpp
@@ -237,14 +237,25 @@ bool FCommand::TryParseChangelist(FStringView ClientSpecName, FStringView Change
 
 			for (const FSourceControlStateRef& FileState : FilesinChangelist)
 			{
-				if (IsPackageFile(FileState->GetFilename()))
+				if (!IsPackageFile(FileState->GetFilename()))
 				{
-					OutPackages.Add(FileState->GetFilename());
+					UE_LOG(LogVirtualizationTool, Log, TEXT("\tIgnoring non-package file '%s'"), *FileState->GetFilename());
+					continue;
 				}
-				else
+
+				if (FileState->IsDeleted())
 				{
-					UE_LOG(LogVirtualizationTool, Log, TEXT("\tIgnoring non-package file '%s'"), *FileState->GetFilename());
+					UE_LOG(LogVirtualizationTool, Verbose, TEXT("\tIgnoring package marked for delete '%s'"), *FileState->GetFilename());
+					continue;
 				}
+
+				if (FileState->IsIgnored())
+				{
+					UE_LOG(LogVirtualizationTool, Verbose, TEXT("\tIgnoring package marked for ignore '%s'"), *FileState->GetFilename());
+					continue;
+				}
+
+				OutPackages.Add(FileState->GetFilename());
 			}
 
 			if (OutChangelist != nullptr)
