[UVT] Add a log line showing when engine pre init ends and the tool begins.

#rnx

diff --git a/Engine/Source/Programs/UnrealVirtualizationTool/Private/UnrealVirtualizationTool.cpp b/Engine/Source/Programs/UnrealVirtualizationTool/Private/UnrealVirtualizationTool.cpp
--- a/Engine/Source/Programs/UnrealVirtualizationTool/Private/UnrealVirtualizationTool.cpp
+++ b/Engine/Source/Programs/UnrealVirtualizationTool/Private/UnrealVirtualizationTool.cpp
@@ -36,6 +36,8 @@ int32 UnrealVirtualizationToolMain(int32 ArgC, TCHAR* ArgV[])
 
 	FModuleManager::Get().StartProcessingNewlyLoadedObjects();
 
+	UE_LOG(LogVirtualizationTool, Display, TEXT("Running UnrealVirtualization Tool"));
+
 	EProcessResult ProcessResult = EProcessResult::Success;
 
 	FUnrealVirtualizationToolApp App;
