[UVT] Initialize the output redirector before engine preinit when run with -MinimalLogging

#rnx

- Some systems running in engine pre init can generate log info with Display verbosity and since we were not adding output directory until after pre init we'd end up with log info in the UVT output that we didn't want.
- Moved FOverrideOutputDevice to it's own code file.
- If a FOverrideOutputDevice is created before the output devices have initialized we need to register with the FCoreDelegates::OnOutputDevicesInit delegate so that we can apply our override once it has been set up.
- Since we want to enable the override before engine preinit we cannot use the engines commandline parsing code. We need our own parsing that can operate and the raw commandline parameters.

diff --git a/Engine/Source/Programs/UnrealVirtualizationTool/Private/OverrideOutputDevice.cpp b/Engine/Source/Programs/UnrealVirtualizationTool/Private/OverrideOutputDevice.cpp
new file mode 100644
--- /dev/null
+++ b/Engine/Source/Programs/UnrealVirtualizationTool/Private/OverrideOutputDevice.cpp
@@ -0,0 +1,85 @@
+// Copyright Epic Games, Inc. All Rights Reserved.
+
+#include "OverrideOutputDevice.h"
+
+#include "Logging/StructuredLog.h"
+#include "Misc/CoreDelegates.h"
+#include "UnrealVirtualizationTool.h"
+
+namespace UE
+{
+
+FOverrideOutputDevice::FOverrideOutputDevice()
+{
+	OriginalLog = GWarn;
+	GWarn = this;
+
+	// If this was created before the output devices were initialized we need to register for a
+	// callback when they are initialized so that we can override them.
+	if (OriginalLog == nullptr)
+	{
+		OnInitHandle = FCoreDelegates::OnOutputDevicesInit.AddLambda([this]()
+			{
+				OriginalLog = GWarn;
+				GWarn = this;
+			});
+	}
+}
+
+FOverrideOutputDevice::~FOverrideOutputDevice()
+{
+	GWarn = OriginalLog;
+
+	if (OnInitHandle.IsValid())
+	{
+		FCoreDelegates::OnOutputDevicesInit.Remove(OnInitHandle);
+		OnInitHandle.Reset();
+	}
+}
+
+void FOverrideOutputDevice::Serialize(const TCHAR* V, ELogVerbosity::Type Verbosity, const FName& Category, double Time)
+{
+#if !NO_LOGGING
+	if (ShouldFilterMessage(Verbosity, Category))
+	{
+		Verbosity = ELogVerbosity::Log;
+	}
+#endif // !NO_LOGGING
+
+	FFeedbackContextAnsi::Serialize(V, Verbosity, Category, Time);
+}
+
+void FOverrideOutputDevice::SerializeRecord(const UE::FLogRecord& Record)
+{
+#if !NO_LOGGING
+	if (ShouldFilterMessage(Record.GetVerbosity(), Record.GetCategory()))
+	{
+		UE::FLogRecord LocalRecord = Record;
+		LocalRecord.SetVerbosity(ELogVerbosity::Log);
+		return FFeedbackContextAnsi::SerializeRecord(LocalRecord);
+	}
+#endif // !NO_LOGGING
+
+	FFeedbackContextAnsi::SerializeRecord(Record);
+}
+
+bool FOverrideOutputDevice::ShouldFilterMessage(ELogVerbosity::Type Verbosity, const FName& Category)
+{
+#if !NO_LOGGING
+	// We only want 'LogVirtualizationTool' messages in display
+	if (Verbosity == ELogVerbosity::Display && Category != LogVirtualizationTool.GetCategoryName())
+	{
+		return true;
+	}
+
+	// Suppress errors from our reporting systems
+	if (Verbosity == ELogVerbosity::Error && Category == LogOutputDevice.GetCategoryName())
+	{
+		return true;
+	}
+#endif //  !NO_LOGGING
+
+	return false;
+}
+
+} // namespace UE
diff --git a/Engine/Source/Programs/UnrealVirtualizationTool/Private/OverrideOutputDevice.h b/Engine/Source/Programs/UnrealVirtualizationTool/Private/OverrideOutputDevice.h
new file mode 100644
--- /dev/null
+++ b/Engine/Source/Programs/UnrealVirtualizationTool/Private/OverrideOutputDevice.h
@@ -0,0 +1,39 @@
+// Copyright Epic Games, Inc. All Rights Reserved.
+
+#include "Delegates/IDelegateInstance.h"
+#include "HAL/FeedbackContextAnsi.h"
+
+namespace UE { class FLogRecord; }
+
+namespace UE
+{
+/**
+ * This class can be used to prevent log messages from other systems being logged with the Display verbosity.
+ * In practical terms this means as long as the class is alive, only LogVirtualizationTool messages will
+ * be logged to the display meaning the user will have less information to deal with.
+ */
+class FOverrideOutputDevice final : public FFeedbackContextAnsi
+{
+public:
+	FOverrideOutputDevice();
+
+	virtual ~FOverrideOutputDevice();
+
+	virtual void Serialize(const TCHAR* V, ELogVerbosity::Type Verbosity, const FName& Category) override
+	{
+		Serialize(V, Verbosity, Category, -1.0);
+	}
+
+	virtual void Serialize(const TCHAR* V, ELogVerbosity::Type Verbosity, const FName& Category, double Time) override;
+
+	virtual void SerializeRecord(const UE::FLogRecord& Record) override;
+
+private:
+
+	bool ShouldFilterMessage(ELogVerbosity::Type Verbosity, const FName& Category);
+
+	FFeedbackContext* OriginalLog;
+	FDelegateHandle OnInitHandle;
+};
+
+} // namespace UE
diff --git a/Engine/Source/Programs/UnrealVirtualizationTool/Private/UnrealVirtualizationTool.cpp b/Engine/Source/Programs/UnrealVirtualizationTool/Private/UnrealVirtualizationTool.cpp
--- a/Engine/Source/Programs/UnrealVirtualizationTool/Private/UnrealVirtualizationTool.cpp
+++ b/Engine/Source/Programs/UnrealVirtualizationTool/Private/UnrealVirtualizationTool.cpp
@@ -7,10 +7,46 @@
 #include "RequiredProgramMainCPPInclude.h"
 #include "UnrealVirtualizationToolApp.h"
 
+#include "OverrideOutputDevice.h"
+
 IMPLEMENT_APPLICATION(UnrealVirtualizationTool, "UnrealVirtualizationTool");
 
 DEFINE_LOG_CATEGORY(LogVirtualizationTool);
 
+/**
+ * Used to search the raw command line to see if it contains a specific switch.
+ * This is the same as calling FParse::Param(FCommandLine::Get(), Switch) except
+ * it will work before FCommandLine has initialized.
+ * 
+ * Note that like the example code above we expect the caller to omit the '-'
+ * character from the switch being searched for.
+ * For example if the user is searching for '-Example' we would expect them to
+ * call this function with TEXT("Example").
+ * 
+ * @param Switch The command line switch to search for
+ * @param ArgC The number of args provided
+ * @param ArgV An array of args provided
+ * 
+ * @return True if the requested switch was found, otherwise false.
+ */
+static bool DoesSwitchExist(FStringView Switch, int32 ArgC, TCHAR* ArgV[])
+{
+	// Skip the first arg as it will just be the path of the exe ( see BuildFromArgVImpl)
+	for (int32 Index = 1; Index < ArgC; ++Index)
+	{
+		// Skip the first character of each arg as each switch should start with '-'
+		FStringView Arg(ArgV[Index]);
+		Arg.RightChopInline(1);
+
+		if (Switch == Arg)
+		{
+			return true;
+		}
+	}
+
+	return false;
+}
+
 int32 UnrealVirtualizationToolMain(int32 ArgC, TCHAR* ArgV[])
 {
 	TRACE_CPUPROFILER_EVENT_SCOPE(UnrealVirtualizationToolMain);
@@ -20,6 +56,12 @@ int32 UnrealVirtualizationToolMain(int32 ArgC, TCHAR* ArgV[])
 	// Allows this program to accept a project argument on the commandline and use project-specific config
 	UE::ProjectUtilities::ParseProjectDirFromCommandline(ArgC, ArgV);
 
+	TUniquePtr<UE::FOverrideOutputDevice> OutputDeviceOverride;
+	if (DoesSwitchExist(TEXTVIEW("MinimalLogging"), ArgC, ArgV))
+	{
+		OutputDeviceOverride = MakeUnique<UE::FOverrideOutputDevice>();
+	}
+
 	GEngineLoop.PreInit(ArgC, ArgV);
 	check(GConfig && GConfig->IsReadyForUse());
 
diff --git a/Engine/Source/Programs/UnrealVirtualizationTool/Private/UnrealVirtualizationToolApp.cpp b/Engine/Source/Programs/UnrealVirtualizationTool/Private/UnrealVirtualizationToolApp.cpp
--- a/Engine/Source/Programs/UnrealVirtualizationTool/Private/UnrealVirtualizationToolApp.cpp
+++ b/Engine/Source/Programs/UnrealVirtualizationTool/Private/UnrealVirtualizationToolApp.cpp
@@ -107,81 +107,6 @@ void CreateChildProcessFilePath(const FGuid& Id, FStringView Extension, FStringB
 	CreateChildProcessFilePath(WriteToString<40>(Id), Extension, OutPath);
 }
 
-
-/**
- * This class can be used to prevent log messages from other systems being logged with the Display verbosity.
- * In practical terms this means as long as the class is alive, only LogVirtualizationTool messages will
- * be logged to the display meaning the user will have less information to deal with.
- */
-class FOverrideOutputDevice final : public FFeedbackContextAnsi
-{
-public:
-	FOverrideOutputDevice()
-	{
-		OriginalLog = GWarn;
-		GWarn = this;
-	}
-
-	virtual ~FOverrideOutputDevice()
-	{
-		GWarn = OriginalLog;
-	}
-
-	virtual void Serialize(const TCHAR* V, ELogVerbosity::Type Verbosity, const FName& Category) override
-	{
-		Serialize(V, Verbosity, Category, -1.0);
-	}
-
-	virtual void Serialize(const TCHAR* V, ELogVerbosity::Type Verbosity, const FName& Category, double Time) override
-	{
-#if !NO_LOGGING
-		if (ShouldFilterMessage(Verbosity, Category))
-		{
-			Verbosity = ELogVerbosity::Log;
-		}
-#endif // !NO_LOGGING
-
-		FFeedbackContextAnsi::Serialize(V, Verbosity, Category, Time);
-	}
-
-	virtual void SerializeRecord(const UE::FLogRecord& Record) override
-	{
-#if !NO_LOGGING
-		if (ShouldFilterMessage(Record.GetVerbosity(), Record.GetCategory()))
-		{
-			UE::FLogRecord LocalRecord = Record;
-			LocalRecord.SetVerbosity(ELogVerbosity::Log);
-			return FFeedbackContextAnsi::SerializeRecord(LocalRecord);
-		}
-#endif // !NO_LOGGING
-
-		FFeedbackContextAnsi::SerializeRecord(Record);
-	}
-
-private:
-
-	bool ShouldFilterMessage(ELogVerbosity::Type Verbosity, const FName& Category)
-	{
-#if !NO_LOGGING
-		// We only want 'LogVirtualizationTool' messages in display
-		if (Verbosity == ELogVerbosity::Display && Category != LogVirtualizationTool.GetCategoryName())
-		{
-			return true;
-		}
-
-		// Suppress errors from our reporting systems
-		if (Verbosity == ELogVerbosity::Error && Category == LogOutputDevice.GetCategoryName())
-		{
-			return true;
-		}
-#endif //  !NO_LOGGING
-
-		return false;
-	}
-
-	FFeedbackContext* OriginalLog;
-};
-
 } // namespace
 
 namespace UE::Virtualization
@@ -447,13 +372,7 @@ EInitResult FUnrealVirtualizationToolApp::TryParseGlobalOptions(const TCHAR* Cmd
 
 	if (FParse::Param(CmdLine, TEXT("MinimalLogging")))
 	{
-		OutputDeviceOverride = MakeUnique<FOverrideOutputDevice>();
-
-		AddGlobalOption(TEXT("-MinimalLogging"));	
-	}
-	else
-	{
-		OutputDeviceOverride.Reset();
+		AddGlobalOption(TEXT("-MinimalLogging"));
 	}
 
 	// Now add commandline switches used in UnrealVirtualizationToolMain (probably should be doing this setup work there)
