CIS Follow up fix for 36778047

#rnx

- Add missing include guard, not entirely sure how that was missed by the original presubmit checks.

diff --git a/Engine/Source/Programs/UnrealVirtualizationTool/Private/OverrideOutputDevice.h b/Engine/Source/Programs/UnrealVirtualizationTool/Private/OverrideOutputDevice.h
--- a/Engine/Source/Programs/UnrealVirtualizationTool/Private/OverrideOutputDevice.h
+++ b/Engine/Source/Programs/UnrealVirtualizationTool/Private/OverrideOutputDevice.h
@@ -3,6 +3,8 @@
 #include "Delegates/IDelegateInstance.h"
 #include "HAL/FeedbackContextAnsi.h"
 
+#pragma once
+
 namespace UE { class FLogRecord; }
 
 namespace UE
