UnrealVirtualizationTool: Fix include issue

#rnx

diff --git a/Engine/Source/Programs/UnrealVirtualizationTool/UnrealVirtualizationTool.Build.cs b/Engine/Source/Programs/UnrealVirtualizationTool/UnrealVirtualizationTool.Build.cs
--- a/Engine/Source/Programs/UnrealVirtualizationTool/UnrealVirtualizationTool.Build.cs
+++ b/Engine/Source/Programs/UnrealVirtualizationTool/UnrealVirtualizationTool.Build.cs
@@ -6,6 +6,8 @@ public class UnrealVirtualizationTool : ModuleRules
 {
 	public UnrealVirtualizationTool(ReadOnlyTargetRules Target) : base(Target)
 	{
+		PrivateIncludePaths.Add(System.IO.Path.Combine(GetModuleDirectory("Projects"), "Internal"));
+
 		PrivateDependencyModuleNames.AddRange(
 			new string[]
 			{
