ProjectUtilities.h needs to be Public

#rnx

diff --git a/Engine/Source/Programs/UnrealVirtualizationTool/UnrealVirtualizationTool.Build.cs b/Engine/Source/Programs/UnrealVirtualizationTool/UnrealVirtualizationTool.Build.cs
--- a/Engine/Source/Programs/UnrealVirtualizationTool/UnrealVirtualizationTool.Build.cs
+++ b/Engine/Source/Programs/UnrealVirtualizationTool/UnrealVirtualizationTool.Build.cs
@@ -6,8 +6,6 @@ public class UnrealVirtualizationTool : ModuleRules
 {
 	public UnrealVirtualizationTool(ReadOnlyTargetRules Target) : base(Target)
 	{
-		PrivateIncludePaths.Add(System.IO.Path.Combine(GetModuleDirectory("Projects"), "Internal"));
-
 		PrivateDependencyModuleNames.AddRange(
 			new string[]
 			{
diff --git a/Engine/Source/Runtime/Projects/Internal/ProjectUtilities.h b/Engine/Source/Runtime/Projects/Internal/ProjectUtilities.h
deleted file mode 100644
diff --git a/Engine/Source/Runtime/Projects/Public/ProjectUtilities.h b/Engine/Source/Runtime/Projects/Public/ProjectUtilities.h
new file mode 100644
--- /dev/null
+++ b/Engine/Source/Runtime/Projects/Public/ProjectUtilities.h
@@ -0,0 +1,13 @@
+// Copyright Epic Games, Inc. All Rights Reserved.
+
+#pragma once
+
+#include "CoreTypes.h"
+
+namespace UE::ProjectUtilities
+{
+	/**
+	 * Allows a monolithic target to control the game directory by the presence of commandline arguments for the project.
+	 */
+	PROJECTS_API void ParseProjectDirFromCommandline(int32 ArgC, TCHAR* ArgV[]);
+}
