[UVT] Allow the rehydrate commandline to work on changelists without a p4client being provided

#rnx

- The virtualize command can work on a changelist without providing the p4client, this command should be able to function in the same way.

diff --git a/Engine/Source/Programs/UnrealVirtualizationTool/Private/Commands/RehydrateCommand.cpp b/Engine/Source/Programs/UnrealVirtualizationTool/Private/Commands/RehydrateCommand.cpp
--- a/Engine/Source/Programs/UnrealVirtualizationTool/Private/Commands/RehydrateCommand.cpp
+++ b/Engine/Source/Programs/UnrealVirtualizationTool/Private/Commands/RehydrateCommand.cpp
@@ -41,7 +41,9 @@ bool FRehydrateCommand::Initialize(const TCHAR* CmdLine)
 
 	ParseCommandLine(CmdLine, Tokens, Switches);
 	
+	FString SourceChangelistNumber;
 	FString SwitchValue;
+
 	for (const FString& Switch :  Switches)
 	{
 		EPathResult Result = ParseSwitchForPaths(Switch, Packages);
@@ -61,7 +63,7 @@ bool FRehydrateCommand::Initialize(const TCHAR* CmdLine)
 		}
 		else if (FParse::Value(*Switch, TEXT("Changelist="), SwitchValue))
 		{
-			return TryParseChangelist(ClientSpecName, SwitchValue, Packages, nullptr);
+			SourceChangelistNumber = SwitchValue;
 		}
 		else if (Switch == TEXT("Checkout"))
 		{
@@ -69,6 +71,40 @@ bool FRehydrateCommand::Initialize(const TCHAR* CmdLine)
 		}
 	}
 
+	// Process the provided changelist if one was found
+	if (!SourceChangelistNumber.IsEmpty())
+	{
+		// If no client spec was provided we need to find it for the changelist
+		// In theory this duplicates a lot of the work found in ::TryParseChangelist
+		// but at the moment the FGetChangelistDetails operation is not compatible
+		// with the FSourceControlChangelistStateRef/FSourceControlStateRef API
+		// so we are stuck with duplication of work.
+		if (ClientSpecName.IsEmpty())
+		{
+			ClientSpecName = FindClientSpecForChangelist(SourceChangelistNumber);
+			if (!ClientSpecName.IsEmpty())
+			{
+				FSourceControlResultInfo Info;
+				if (SCCProvider->SwitchWorkspace(ClientSpecName, Info, nullptr) != ECommandResult::Succeeded)
+				{
+					UE_LOG(LogVirtualization, Error, TEXT("Failed to switch to workspace '%s'"), *ClientSpecName);
+					return false;
+				}
+			}
+			else
+			{
+				UE_LOG(LogVirtualization, Error, TEXT("Count not find a valid workspace for the changelist '%s'"), *SourceChangelistNumber);
+				return false;
+			}
+		}
+
+		if (!TryParseChangelist(ClientSpecName, SourceChangelistNumber, Packages, nullptr))
+		{
+			UE_LOG(LogVirtualization, Error, TEXT("Failed to find the files in the changelist '%s'"), *SourceChangelistNumber);
+			return false;
+		}
+	}
+
 	return true;
 }
 
