[UVT] Improvements to the virtualization process logging.

#rnx

- Once packages are assigned to projects we now log each project found and the number of packages assigned to it. If some packages were not found to be in any project then we explicitly log how many and that they will be ignored.
- Made it clearer when the virtualization process is actually running.
-- 'Checking package(s) for the project..' -> 'Running the virtualization process for the project...'

diff --git a/Engine/Source/Programs/UnrealVirtualizationTool/Private/Commands/VirtualizeCommand.cpp b/Engine/Source/Programs/UnrealVirtualizationTool/Private/Commands/VirtualizeCommand.cpp
--- a/Engine/Source/Programs/UnrealVirtualizationTool/Private/Commands/VirtualizeCommand.cpp
+++ b/Engine/Source/Programs/UnrealVirtualizationTool/Private/Commands/VirtualizeCommand.cpp
@@ -151,7 +151,7 @@ bool FVirtualizeCommand::ProcessProject(const FProject& Project, TUniquePtr<FCom
 	TStringBuilder<128> ProjectName;
 	ProjectName << Project.GetProjectName();
 
-	UE_LOG(LogVirtualizationTool, Display, TEXT("\tChecking package(s) for the project '%s'..."), ProjectName.ToString());
+	UE_LOG(LogVirtualizationTool, Display, TEXT("\tRunning the virtualization process for the project '%s'..."), ProjectName.ToString());
 
 	FConfigFile EngineConfigWithProject;
 	if (!Project.TryLoadConfig(EngineConfigWithProject))
diff --git a/Engine/Source/Programs/UnrealVirtualizationTool/Private/ProjectFiles.cpp b/Engine/Source/Programs/UnrealVirtualizationTool/Private/ProjectFiles.cpp
--- a/Engine/Source/Programs/UnrealVirtualizationTool/Private/ProjectFiles.cpp
+++ b/Engine/Source/Programs/UnrealVirtualizationTool/Private/ProjectFiles.cpp
@@ -159,6 +159,18 @@ TArray<FString> FProject::GetAllPackages() const
 	return Packages;
 }
 
+int32 FProject::GetNumPackages() const
+{
+	int32 PackageCount = PackagePaths.Num();
+
+	for (const FPlugin& Plugin : Plugins)
+	{
+		PackageCount += Plugin.PackagePaths.Num();
+	}
+
+	return PackageCount;
+}
+
 void FProject::RegisterMountPoints() const
 {
 	TRACE_CPUPROFILER_EVENT_SCOPE(FProject::RegisterMountPoints);
diff --git a/Engine/Source/Programs/UnrealVirtualizationTool/Private/ProjectFiles.h b/Engine/Source/Programs/UnrealVirtualizationTool/Private/ProjectFiles.h
--- a/Engine/Source/Programs/UnrealVirtualizationTool/Private/ProjectFiles.h
+++ b/Engine/Source/Programs/UnrealVirtualizationTool/Private/ProjectFiles.h
@@ -39,6 +39,7 @@ public:
 	FStringView GetProjectRoot() const;
 
 	TArray<FString> GetAllPackages() const;
+	int32 GetNumPackages() const;
 
 	void RegisterMountPoints() const;
 	void UnRegisterMountPoints() const;
diff --git a/Engine/Source/Programs/UnrealVirtualizationTool/Private/UnrealVirtualizationToolApp.cpp b/Engine/Source/Programs/UnrealVirtualizationTool/Private/UnrealVirtualizationToolApp.cpp
--- a/Engine/Source/Programs/UnrealVirtualizationTool/Private/UnrealVirtualizationToolApp.cpp
+++ b/Engine/Source/Programs/UnrealVirtualizationTool/Private/UnrealVirtualizationToolApp.cpp
@@ -155,7 +155,7 @@ EInitResult FUnrealVirtualizationToolApp::Initialize()
 	{
 		TArray<FString> Packages = CurrentCommand->GetPackages();
 
-		UE_LOG(LogVirtualizationTool, Display, TEXT("\tFound '%d' package file(s)"), Packages.Num());
+		UE_LOG(LogVirtualizationTool, Display, TEXT("\tFound %d package file(s)"), Packages.Num());
 
 		if (!TrySortFilesByProject(Packages))
 		{
@@ -448,7 +448,19 @@ bool FUnrealVirtualizationToolApp::TrySortFilesByProject(const TArray<FString>&
 		}
 	}
 
-	UE_LOG(LogVirtualizationTool, Display, TEXT("\tThe package files are associated with '%d' projects(s)"), Projects.Num());
+	UE_LOG(LogVirtualizationTool, Display, TEXT("\tFound the following project(s):"));
+
+	int32 TotalPackagesAssigned = 0;
+	for (const FProject& Project : Projects)
+	{
+		TotalPackagesAssigned += Project.GetNumPackages();
+		UE_LOGFMT(LogVirtualizationTool, Display, "\t\t{ProjectName}: {NumPackages} package(s)", Project.GetProjectName(), Project.GetNumPackages());
+	}
+
+	if (TotalPackagesAssigned != Packages.Num())
+	{
+		UE_LOG(LogVirtualizationTool, Display, TEXT("\tCould not find a project for %d package(s) which will be ignored"), Packages.Num() - TotalPackagesAssigned);
+	}
 
 	return true;
 }
