CIS: Fix UnrealVirtualizationTool linker warnings when building with clang.

#rnx

- Removed some features from UVT to prevent a 'lld-link: warning: found both wmain and main; using latter' linker warning.
- Ideally we need to remove 'TypedElementFramework' from being included by the PerforceSourceControl module but that will be it's own work item.
- This current fix will also provide a minor reduction in exe/pdb size.

diff --git a/Engine/Source/Programs/UnrealVirtualizationTool/UnrealVirtualizationTool.Target.cs b/Engine/Source/Programs/UnrealVirtualizationTool/UnrealVirtualizationTool.Target.cs
--- a/Engine/Source/Programs/UnrealVirtualizationTool/UnrealVirtualizationTool.Target.cs
+++ b/Engine/Source/Programs/UnrealVirtualizationTool/UnrealVirtualizationTool.Target.cs
@@ -24,6 +24,10 @@ public class UnrealVirtualizationToolTarget : TargetRules
 		bCompileAgainstEngine = false;
 		bUsesSlate = false;
 
+		// We don't use these and they are causing 'lld-link: warning: found both wmain and main; using latter' warnings with clang
+		bCompileFreeType = false;
+		bCompileICU = false;
+
 		// This app is a console application, not a Windows app (sets entry point to main(), instead of WinMain())
 		bIsBuildingConsoleApplication = true;
 
