UVT - Move everything in UnrealVirtualizationToolApp.cpp into the UE::Virtualization namespace.

#rnx

diff --git a/Engine/Source/Programs/UnrealVirtualizationTool/Private/UnrealVirtualizationToolApp.cpp b/Engine/Source/Programs/UnrealVirtualizationTool/Private/UnrealVirtualizationToolApp.cpp
--- a/Engine/Source/Programs/UnrealVirtualizationTool/Private/UnrealVirtualizationToolApp.cpp
+++ b/Engine/Source/Programs/UnrealVirtualizationTool/Private/UnrealVirtualizationToolApp.cpp
@@ -25,31 +25,31 @@
 #include "VirtualizationUtilities.h"
 #include "Virtualization/VirtualizationSystem.h"
 
-namespace
+namespace UE::Virtualization
 {
 
 /** Utility to get EMode from a string */
-void LexFromString(UE::Virtualization::EMode& OutValue, const FStringView& InString)
+void LexFromString(EMode& OutValue, const FStringView& InString)
 {
 	if (InString == TEXT("Changelist"))
 	{
-		OutValue = UE::Virtualization::EMode::Changelist;
+		OutValue = EMode::Changelist;
 	}
 	else if (InString == TEXT("PackageList"))
 	{
-		OutValue = UE::Virtualization::EMode::PackageList;
+		OutValue = EMode::PackageList;
 	}
 	else if (InString == TEXT("Virtualize"))
 	{
-		OutValue = UE::Virtualization::EMode::Virtualize;
+		OutValue = EMode::Virtualize;
 	}
 	else if (InString == TEXT("Rehydrate"))
 	{
-		OutValue = UE::Virtualization::EMode::Rehydrate;
+		OutValue = EMode::Rehydrate;
 	}
 	else
 	{
-		OutValue = UE::Virtualization::EMode::Unknown;
+		OutValue = EMode::Unknown;
 	}
 }
 
@@ -57,10 +57,8 @@ void LexFromString(UE::Virtualization::EMode& OutValue, const FStringView& InStr
  * Utility to convert EMode to a string. NOTE that legacy modes will return the
  * name of the newer mode that replaces it and not the legacy mode.
  */
-const TCHAR* LexToString(UE::Virtualization::EMode Mode)
+const TCHAR* LexToString(EMode Mode)
 {
-	using namespace UE::Virtualization;
-
 	switch (Mode)
 	{
 		case EMode::Unknown:
@@ -107,11 +105,6 @@ void CreateChildProcessFilePath(const FGuid& Id, FStringView Extension, FStringB
 	CreateChildProcessFilePath(WriteToString<40>(Id), Extension, OutPath);
 }
 
-} // namespace
-
-namespace UE::Virtualization
-{
-
 FUnrealVirtualizationToolApp::FUnrealVirtualizationToolApp()
 {
 
