UVT - Improve logging when '-SourceProject=' is used.

- We now log when the project path is overriden by the cmdline arg 'SourceProject', we also log if the project cannot be overriden because the project path cannot be found.
- We now normalize for file path not directory path. The result will be the same but gives better context.
- We also call FPaths::RemoveDuplicateSlashes on the provided path. This can sometimes clean up issues caused by incorrectly concatinated strings in build scripts.

diff --git a/Engine/Source/Programs/UnrealVirtualizationTool/Private/ProjectFiles.cpp b/Engine/Source/Programs/UnrealVirtualizationTool/Private/ProjectFiles.cpp
--- a/Engine/Source/Programs/UnrealVirtualizationTool/Private/ProjectFiles.cpp
+++ b/Engine/Source/Programs/UnrealVirtualizationTool/Private/ProjectFiles.cpp
@@ -357,13 +357,17 @@ void FProject::ProcessProjectPath()
 	// The commandline takes priority and can be applied to any project type
 	if (FParse::Value(FCommandLine::Get(), TEXT("-SourceProject="), OverridenProjectFilePath))
 	{
-		FPaths::NormalizeDirectoryName(OverridenProjectFilePath);
+		FPaths::NormalizeFilename(OverridenProjectFilePath);
+		FPaths::RemoveDuplicateSlashes(OverridenProjectFilePath);
 
-		if (!IFileManager::Get().FileExists(*OverridenProjectFilePath))
+		if (IFileManager::Get().FileExists(*OverridenProjectFilePath))
 		{
-			// TODO Give Error
+			UE_LOGFMT(LogVirtualizationTool, Display, "\t\t({ProjectName}) Overriding project path with '{ProjectFilePath}' due to command arg '-SourceProject'", GetProjectName(), OverridenProjectFilePath);
+		}
+		else
+		{
+			UE_LOGFMT(LogVirtualizationTool, Error, "Unable to find project file {ProjectFilePath} from command arg '-SourceProject'", OverridenProjectFilePath);
 			OverridenProjectFilePath = ProjectFilePath;
-			return;
 		}
 
 		return;
@@ -381,7 +385,7 @@ void FProject::ProcessProjectPath()
 		return;
 	}
 
-	UE_LOGFMT(LogVirtualization, Error, "Unable to find a valid GameProject for {ProjectFilePath}", ProjectFilePath);
+	UE_LOGFMT(LogVirtualizationTool, Error, "Unable to find a valid GameProject for {ProjectFilePath}", ProjectFilePath);
 }
 
 } // namespace UE::Virtualization
