[UVT] Tell the user how many packages were virtualized and how much disk space has been saved.

#rnx

- Add an exxplicit log message if no packages are virtualized.
- We only report the number of packages checked out from revision control if at least one package was virtualized.

diff --git a/Engine/Source/Programs/UnrealVirtualizationTool/Private/Commands/VirtualizeCommand.cpp b/Engine/Source/Programs/UnrealVirtualizationTool/Private/Commands/VirtualizeCommand.cpp
--- a/Engine/Source/Programs/UnrealVirtualizationTool/Private/Commands/VirtualizeCommand.cpp
+++ b/Engine/Source/Programs/UnrealVirtualizationTool/Private/Commands/VirtualizeCommand.cpp
@@ -210,9 +210,22 @@ bool FVirtualizeCommand::ProcessProject(const FProject& Project, TUniquePtr<FCom
 		return false;
 	}
 
-	if (bShouldCheckout)
+	if (!Result.VirtualizedPackages.IsEmpty())
+	{
+		UE_LOG(LogVirtualizationTool, Display, TEXT("\t\t%d packages were virtualized"), Result.VirtualizedPackages.Num());
+		UE_LOG(LogVirtualizationTool, Display, TEXT("\t\tTotal Reduction %s (%s -> %s)"),
+			*FText::AsMemory(Result.PreOperationSize - Result.PostOperationSize).ToString(),
+			*FText::AsMemory(Result.PreOperationSize).ToString(),
+			*FText::AsMemory(Result.PostOperationSize).ToString());
+
+		if (bShouldCheckout)
+		{
+			UE_LOG(LogVirtualizationTool, Display, TEXT("\t\t%d packages were checked out of revision control"), Result.CheckedOutPackages.Num());
+		}
+	}
+	else
 	{
-		UE_LOG(LogVirtualizationTool, Display, TEXT("\t\t%d packages were checked out of revision control"), Result.CheckedOutPackages.Num());
+		UE_LOG(LogVirtualizationTool, Display, TEXT("\t\tNo package files were virtualized"));
 	}
 
 	UE_LOG(LogVirtualizationTool, Display, TEXT("\t\tTime taken %.2f(s)"), Result.TimeTaken);
