[UVT] Tell the user how many packages were rehydrated and how much disk space this will take up.

#rnx

- Add an exxplicit log message if no packages are rehydrated.
- We only report the number of packages checked out from revision control if at least one package was rehydrated.

diff --git a/Engine/Source/Programs/UnrealVirtualizationTool/Private/Commands/RehydrateCommand.cpp b/Engine/Source/Programs/UnrealVirtualizationTool/Private/Commands/RehydrateCommand.cpp
--- a/Engine/Source/Programs/UnrealVirtualizationTool/Private/Commands/RehydrateCommand.cpp
+++ b/Engine/Source/Programs/UnrealVirtualizationTool/Private/Commands/RehydrateCommand.cpp
@@ -175,9 +175,22 @@ bool FRehydrateCommand::ProcessProject(const FProject& Project, TUniquePtr<FComm
 		return false;
 	}
 
-	if (bShouldCheckout)
+	if (!Result.RehydratedPackages.IsEmpty())
+	{
+		UE_LOG(LogVirtualizationTool, Display, TEXT("\t\t%d packages were hydrated"), Result.RehydratedPackages.Num());
+		UE_LOG(LogVirtualizationTool, Display, TEXT("\t\tTotal hydration increase %s (%s -> %s)"),
+			*FText::AsMemory(Result.PostOperationSize - Result.PreOperationSize).ToString(),
+			*FText::AsMemory(Result.PreOperationSize).ToString(),
+			*FText::AsMemory(Result.PostOperationSize).ToString());
+
+		if (bShouldCheckout)
+		{
+			UE_LOG(LogVirtualizationTool, Display, TEXT("\t\t%d packages were checked out of revision control"), Result.CheckedOutPackages.Num());
+		}
+	}
+	else
 	{
-		UE_LOG(LogVirtualizationTool, Display, TEXT("\t\t%d packages were checked out of revision control"), Result.CheckedOutPackages.Num());
+		UE_LOG(LogVirtualizationTool, Display, TEXT("\t\tNo package files were hydrated"));
 	}
 
 	UE_LOG(LogVirtualizationTool, Display, TEXT("\t\tTime taken %.2f(s)"), Result.TimeTaken);
