Fixed printf.

diff --git a/Engine/Source/Programs/UnrealVirtualizationTool/Private/Commands/RehydrateCommand.cpp b/Engine/Source/Programs/UnrealVirtualizationTool/Private/Commands/RehydrateCommand.cpp
--- a/Engine/Source/Programs/UnrealVirtualizationTool/Private/Commands/RehydrateCommand.cpp
+++ b/Engine/Source/Programs/UnrealVirtualizationTool/Private/Commands/RehydrateCommand.cpp
@@ -151,7 +151,7 @@ bool FRehydrateCommand::ProcessProject(const FProject& Project, TUniquePtr<FComm
 		UE::Virtualization::Shutdown();
 	};
 
-	UE_LOG(LogVirtualizationTool, Display, TEXT("\t\tAttempting to rehydrate packages..."), ProjectName.ToString());
+	UE_LOG(LogVirtualizationTool, Display, TEXT("\t\tAttempting to rehydrate packages..."));
 
 	const TArray<FString> ProjectPackages = Project.GetAllPackages();
 
