Fixed printf.

diff --git a/Engine/Source/Programs/UnrealVirtualizationTool/Private/Commands/RehydrateCommand.cpp b/Engine/Source/Programs/UnrealVirtualizationTool/Private/Commands/RehydrateCommand.cpp
--- a/Engine/Source/Programs/UnrealVirtualizationTool/Private/Commands/RehydrateCommand.cpp
+++ b/Engine/Source/Programs/UnrealVirtualizationTool/Private/Commands/RehydrateCommand.cpp
@@ -194,7 +194,7 @@ bool FRehydrateCommand::ProcessProject(const FProject& Project, TUniquePtr<FComm
 	}
 
 	UE_LOG(LogVirtualizationTool, Display, TEXT("\t\tTime taken %.2f(s)"), Result.TimeTaken);
-	UE_LOG(LogVirtualizationTool, Display, TEXT("\t\tRehyration of project packages complete!"), ProjectName.ToString());
+	UE_LOG(LogVirtualizationTool, Display, TEXT("\t\tRehyration of project packages complete!"));
 
 	return true;
 }
