Fixed printf.

diff --git a/Engine/Source/Programs/UnrealVirtualizationTool/Private/Commands/VirtualizeCommand.cpp b/Engine/Source/Programs/UnrealVirtualizationTool/Private/Commands/VirtualizeCommand.cpp
--- a/Engine/Source/Programs/UnrealVirtualizationTool/Private/Commands/VirtualizeCommand.cpp
+++ b/Engine/Source/Programs/UnrealVirtualizationTool/Private/Commands/VirtualizeCommand.cpp
@@ -229,7 +229,7 @@ bool FVirtualizeCommand::ProcessProject(const FProject& Project, TUniquePtr<FCom
 	}
 
 	UE_LOG(LogVirtualizationTool, Display, TEXT("\t\tTime taken %.2f(s)"), Result.TimeTaken);
-	UE_LOG(LogVirtualizationTool, Display, TEXT("\t\tVirtualization of project packages complete!"), ProjectName.ToString());
+	UE_LOG(LogVirtualizationTool, Display, TEXT("\t\tVirtualization of project packages complete!"));
 
 	Output = MakeUnique<FVirtualizeCommandOutput>(Project.GetProjectName(), Result.DescriptionTags);
 
